Moving Submit Tool to Source/Programs

[RN] [Experimental] Publishing Submit Tool code

diff --git a/Engine/Source/Programs/SubmitTool/Private/CommandLine/CmdLineParameter.h b/Engine/Source/Programs/SubmitTool/Private/CommandLine/CmdLineParameter.h
new file mode 100644
--- /dev/null
+++ b/Engine/Source/Programs/SubmitTool/Private/CommandLine/CmdLineParameter.h
@@ -0,0 +1,47 @@
+// Copyright Epic Games, Inc. All Rights Reserved.
+
+#pragma once
+
+#include "CoreMinimal.h"
+
+class FCmdLineParameter
+{
+public:
+	FCmdLineParameter() = delete;
+	FCmdLineParameter(FString InKey, bool InIsRequired, FString InDescription, TFunction<bool(const FString& InValue)> InValidator = nullptr, TFunction<void(FString& InValue)> InParseInPlace = nullptr) :
+		Key(InKey),
+		bIsRequired(InIsRequired),
+		Description(InDescription),
+		Validator(InValidator),
+		Parser(InParseInPlace)
+	{}
+
+	FString ParameterKey() const { return this->Key; }
+	FString ParameterDescription() const { return this->Description; }
+	bool IsRequired() const { return this->bIsRequired; }
+	bool IsValid(const FString& value) const
+	{
+		if(Validator != nullptr)
+		{
+			return Validator(value);
+		}
+
+		return true;
+	}
+
+	void CustomParse(FString& InOutValue) const
+	{
+		if(Parser != nullptr)
+		{
+			Parser(InOutValue);
+		}
+	}
+
+
+private:
+	FString Key;
+	bool bIsRequired;
+	FString Description;
+	TFunction<bool(const FString& InParam)> Validator;
+	TFunction<void(FString& InParam)> Parser;
+};
\ No newline at end of file
diff --git a/Engine/Source/Programs/SubmitTool/Private/CommandLine/CmdLineParameters.cpp b/Engine/Source/Programs/SubmitTool/Private/CommandLine/CmdLineParameters.cpp
new file mode 100644
--- /dev/null
+++ b/Engine/Source/Programs/SubmitTool/Private/CommandLine/CmdLineParameters.cpp
@@ -0,0 +1,127 @@
+// Copyright Epic Games, Inc. All Rights Reserved.
+
+#include "CmdLineParameters.h"
+#include "Logging/SubmitToolLog.h"
+#include "Misc/CommandLine.h"
+
+const FCmdLineParameters FCmdLineParameters::Instance = FCmdLineParameters();
+
+FCmdLineParameters::FCmdLineParameters()
+	: Parameters(FSubmitToolCmdLine::SubmitToolCmdLineArgs)
+{
+}
+
+bool FCmdLineParameters::ValidateParameters() const
+{
+	const TCHAR* CommandLine = FCommandLine::Get();
+
+	bool bIsValid = true;
+	for (const TSharedPtr<FCmdLineParameter>& Parameter : Parameters)
+	{
+		if (Parameter->IsRequired())
+		{
+			FString Key = Parameter->ParameterKey();
+			if(!Key.EndsWith(TEXT(" ")))
+			{
+				Key += TEXT(" ");
+			}
+
+			FString Val;
+			if (FParse::Value(CommandLine, *Key, Val))
+			{
+				if (Val.IsEmpty())
+				{
+					UE_LOG(LogSubmitTool, Error, TEXT("Command Line argument '-%s' has no value."), *Parameter->ParameterKey());
+					bIsValid = false;
+				}
+				else if(!Parameter->IsValid(Val))
+				{
+					UE_LOG(LogSubmitTool, Error, TEXT("Command Line argument '-%s' value '%s' is invalid."), *Parameter->ParameterKey(), *Val);
+					bIsValid = false;
+				}
+			}
+			else
+			{
+				UE_LOG(LogSubmitTool, Error, TEXT("Command Line missing '-%s' argument."), *Parameter->ParameterKey());
+				bIsValid = false;
+			}
+		}
+		else
+		{
+			// optional parameters must contain a value if they are present in the command line
+			FString Val;
+			if (FParse::Value(CommandLine, *Parameter->ParameterKey(), Val))
+			{
+				if (Val.IsEmpty())
+				{
+					UE_LOG(LogSubmitTool, Error, TEXT("Command Line argument '-%s' has no value."), *Parameter->ParameterKey());
+					bIsValid = false;
+				}
+			}
+		}
+	}
+
+	return bIsValid;
+}
+
+void FCmdLineParameters::LogParameters() const
+{
+	for (const TSharedPtr<FCmdLineParameter>& Parameter : Parameters)
+	{
+		UE_LOG(LogSubmitTool, Warning, TEXT("-%s\t%s"), *Parameter->ParameterKey(), *Parameter->ParameterDescription());
+	}
+}
+
+const FCmdLineParameters& FCmdLineParameters::Get()
+{
+	return Instance;
+}
+
+bool FCmdLineParameters::Contains(const FString& InKey) const
+{
+	return FParse::Param(FCommandLine::Get(), *InKey);
+}
+
+bool FCmdLineParameters::GetValue(const FString& InKey, FString& OutValue) const
+{
+	const TCHAR* CommandLine = FCommandLine::Get();
+
+	TArray<FString> Tokens;
+	TArray<FString> Switches;
+	FCommandLine::Parse(CommandLine, Tokens, Switches);
+	const TSharedPtr<FCmdLineParameter>* Definition = Parameters.FindByPredicate([InKey](const TSharedPtr<FCmdLineParameter>& InCmdParam) { return InCmdParam->ParameterKey() == InKey; });
+
+	for(int i = 0; i < Tokens.Num(); i++)
+	{
+		if(Tokens[i].Equals(InKey, ESearchCase::IgnoreCase) && i + 1 < Tokens.Num())
+		{
+			OutValue = Tokens[i + 1];
+			if(Definition)
+			{
+				(*Definition)->CustomParse(OutValue);
+			}
+			return true;
+		}
+	}
+
+	for(int i = 0; i < Switches.Num(); i++)
+	{
+		FString Switch = Switches[i];
+
+		if(Switch.StartsWith(InKey, ESearchCase::IgnoreCase))
+		{
+			TArray<FString> SplitSwitch;
+			if(2 == Switch.ParseIntoArray(SplitSwitch, TEXT("="), true))
+			{
+				OutValue = SplitSwitch[1];
+				if(Definition)
+				{
+					(*Definition)->CustomParse(OutValue);
+				}
+				return true;
+			}
+		}
+	}
+
+	return false;
+}
diff --git a/Engine/Source/Programs/SubmitTool/Private/Logic/ChangelistService.cpp b/Engine/Source/Programs/SubmitTool/Private/Logic/ChangelistService.cpp
new file mode 100644
--- /dev/null
+++ b/Engine/Source/Programs/SubmitTool/Private/Logic/ChangelistService.cpp
@@ -0,0 +1,723 @@
+// Copyright Epic Games, Inc. All Rights Reserved.
+
+#include "ChangelistService.h"
+
+#include "Configuration/Configuration.h"
+#include "Logging/SubmitToolLog.h"
+#include "Logic/Services/Interfaces/ISTSourceControlService.h"
+#include "Models/ModelInterface.h"
+#include "Telemetry/TelemetryService.h"
+#include "CommandLine/CmdLineParameters.h"
+
+#include "ISourceControlProvider.h"
+#include "ISourceControlOperation.h"
+#include "SourceControlOperations.h"
+#include "SourceControlResultInfo.h"
+
+THIRD_PARTY_INCLUDES_START
+// used to retrieve the depot paths, UE hides all this p4 specific information under PerforceSourceControl private definitions
+#include <p4/clientapi.h>
+THIRD_PARTY_INCLUDES_END
+
+namespace ChangelistServiceConstants {
+	constexpr float TickDelay = 5.f;
+}
+
+FChangelistService::FChangelistService(
+	const FGeneralParameters& InParameters,
+	const TSharedPtr<ISTSourceControlService> InPerforceService,
+	const FOnChangeListReadyDelegate& InCLReadyCallback,
+	const FOnChangelistRefreshDelegate& InCLRefreshCallback) :
+	Parameters(InParameters),
+	CLReadyCallback(InCLReadyCallback),
+	CLRefreshCallback(InCLRefreshCallback),
+	TickHandle(FTSTicker::GetCoreTicker().AddTicker(FTickerDelegate::CreateRaw(this, &FChangelistService::P4Tick), ChangelistServiceConstants::TickDelay)),
+	SourceControlService(InPerforceService)
+{
+	FCmdLineParameters::Get().GetValue(FSubmitToolCmdLine::P4ChangeList, CLID);
+	Init();
+}
+
+FChangelistService::~FChangelistService()
+{
+	OnCLDescriptionUpdated.Clear();
+}
+
+void FChangelistService::Init()
+{
+	if(SourceControlService->GetProvider().IsValid())
+	{
+		UE::Tasks::Launch(UE_SOURCE_LOCATION, [this] {
+			GetFilesDepotPaths(true);
+		});
+
+		UE::Tasks::Launch(UE_SOURCE_LOCATION, [this] {
+			GetShelvedFilesDepotPaths(true);
+		});
+
+		CachedSCCProvider = SourceControlService->GetProvider().Get();
+
+		if(!CLID.IsNumeric() && CLID.Equals(TEXT("default"), ESearchCase::IgnoreCase))
+		{
+			CreateCLFromDefaultCL();
+		}
+		else
+		{
+			FindInitialChangelistsAsync();
+		}
+		
+	}
+	else
+	{
+		UE_LOG(LogSubmitToolP4, Error, TEXT("Perforce Connection was invalid"));
+	}
+}
+
+void FChangelistService::FindInitialChangelistsAsync()
+{
+	TSharedRef<FUpdatePendingChangelistsStatus> UpdatePendingChangelistsOperation = ISourceControlOperation::Create<FUpdatePendingChangelistsStatus>();
+	UpdatePendingChangelistsOperation->SetUpdateAllChangelists(true);
+
+	UE_LOG(LogSubmitToolP4Debug, Log, TEXT("Finding available changelists"));
+
+	ActiveP4Operations.Add(UpdatePendingChangelistsOperation);
+
+	CachedSCCProvider->Execute(UpdatePendingChangelistsOperation, EConcurrency::Asynchronous, FSourceControlOperationComplete::CreateLambda([this](const FSourceControlOperationRef& UpdateOperation, ECommandResult::Type Result)
+		{
+			if(Result == ECommandResult::Succeeded)
+			{
+				TArray<FSourceControlChangelistRef> Changelists = CachedSCCProvider->GetChangelists(EStateCacheUsage::Use);
+				for(const FSourceControlChangelistRef& CL : Changelists)
+				{
+					if(CL->GetIdentifier() == CLID)
+					{
+						ChangelistPtr = CL;
+					}
+				}
+			}
+			else if(Result == ECommandResult::Cancelled)
+			{
+				UE_LOG(LogSubmitToolP4, Warning, TEXT("Cancelled finding available changelists"));
+			}
+			else if(Result == ECommandResult::Failed)
+			{
+				UE_LOG(LogSubmitToolP4, Error, TEXT("Failed to find available changelists."));
+				PrintErrorMessages(UpdateOperation->GetResultInfo());
+			}
+
+			ActiveP4Operations.Remove(UpdateOperation);
+
+			if(ChangelistPtr)
+			{
+				// Get state for our CL
+				FetchChangelistDataAsync();
+
+				// Fire & Forget updating all CLs
+				TSharedRef<FUpdatePendingChangelistsStatus> UpdatePendingChangelistsOperation = ISourceControlOperation::Create<FUpdatePendingChangelistsStatus>();
+				UpdatePendingChangelistsOperation->SetUpdateFilesStates(true);
+				UpdatePendingChangelistsOperation->SetUpdateShelvedFilesStates(true);
+				UpdatePendingChangelistsOperation->SetUpdateAllChangelists(true);
+				CachedSCCProvider->Execute(UpdatePendingChangelistsOperation, EConcurrency::Asynchronous);
+			}
+			else
+			{
+				UE_LOG(LogSubmitToolP4, Error, TEXT("Couldn't retrieve information from CL %s"), *CLID);
+				FModelInterface::SetErrorState();
+			}
+		}));
+}
+
+void FChangelistService::FetchChangelistDataAsync()
+{
+	TSharedRef<FUpdatePendingChangelistsStatus> UpdatePendingChangelistsOperation = ISourceControlOperation::Create<FUpdatePendingChangelistsStatus>();
+	UpdatePendingChangelistsOperation->SetUpdateFilesStates(true);
+	UpdatePendingChangelistsOperation->SetUpdateShelvedFilesStates(true);
+
+	if(ChangelistPtr.IsValid())
+	{
+		UE_LOG(LogSubmitToolP4Debug, Log, TEXT("Updating CL %s changes from P4"), *GetCLID());
+		UpdatePendingChangelistsOperation.Get().SetChangelistsToUpdate(TArray{ ChangelistPtr.ToSharedRef() });
+	}
+	else
+	{
+		UE_LOG(LogSubmitToolP4Debug, Log, TEXT("Retrieving CL %s information from P4"), *CLID);
+		UpdatePendingChangelistsOperation->SetUpdateAllChangelists(true);
+	}
+
+	ActiveP4Operations.Add(UpdatePendingChangelistsOperation);
+
+	CachedSCCProvider->Execute(UpdatePendingChangelistsOperation, EConcurrency::Asynchronous, FSourceControlOperationComplete::CreateRaw(this, &FChangelistService::FetchChangelistCallback));
+}
+
+
+void FChangelistService::FetchChangelistCallback(const FSourceControlOperationRef& UpdateOperation, ECommandResult::Type Result)
+{
+	if(Result == ECommandResult::Succeeded)
+	{
+		UE_LOG(LogSubmitToolP4Debug, Log, TEXT("Update operation succeeded."));
+		PrintMessages(UpdateOperation->GetResultInfo());
+
+		// If this is the first time we fetch results from perforce, we cache our CL State object and let the app know 
+		if(!ChangelistState.IsValid())
+		{
+			TArray<FSourceControlChangelistRef> Changelists = CachedSCCProvider->GetChangelists(EStateCacheUsage::Use);
+			for(const FSourceControlChangelistRef& CL : Changelists)
+			{
+				if(CL->GetIdentifier() == CLID)
+				{
+					ChangelistState = CachedSCCProvider->GetState(ChangelistPtr.ToSharedRef(), EStateCacheUsage::Use);
+				}
+			}
+
+			if(!ChangelistPtr.IsValid())
+			{
+				UE_LOG(LogSubmitToolP4, Error, TEXT("Failed to fetch CL %s from P4"), *CLID);
+			}
+			else
+			{
+				OriginalDescription = ChangelistState->GetDescriptionText();
+				CLDescription = ChangelistState->GetDescriptionText().ToString();
+				FilesInCL = ChangelistState->GetFilesStates();
+				ShelvedFilesInCL = ChangelistState->GetShelvedFilesStates();
+				PrintFilesAndShelvedFiles();
+			}
+
+			CLReadyCallback.ExecuteIfBound(ChangelistState.IsValid());
+		}
+		else
+		{
+			RehydrateDataFromP4State();
+		}
+	}
+	else if(Result == ECommandResult::Cancelled)
+	{
+		UE_LOG(LogSubmitToolP4, Warning, TEXT("Update operation cancelled."));
+	}
+	else if(Result == ECommandResult::Failed)
+	{
+		UE_LOG(LogSubmitToolP4, Warning, TEXT("Update operation failed."));
+		PrintErrorMessages(UpdateOperation->GetResultInfo());
+	}
+
+	ActiveP4Operations.Remove(UpdateOperation);
+}
+
+void FChangelistService::RevertUnchangedFilesAsync(const FSourceControlOperationComplete& OnRevertComplete)
+{
+	UE_LOG(LogSubmitToolP4, Log, TEXT("Reverting unchanged files from CL %s..."), *GetCLID());
+
+	TSharedRef<FRevertUnchanged> RevertUnchangedOperation = ISourceControlOperation::Create<FRevertUnchanged>();
+
+	ActiveP4Operations.Add(RevertUnchangedOperation);
+
+	CachedSCCProvider->Execute(RevertUnchangedOperation, ChangelistPtr, EConcurrency::Asynchronous, FSourceControlOperationComplete::CreateLambda(
+		[this, OnRevertComplete](const FSourceControlOperationRef& Operation, ECommandResult::Type Result) {
+
+			// Revert Unchanged returns as failed if there were no files to revert, check ErrorMessages to see actual failures
+			if(Result == ECommandResult::Cancelled)
+			{
+				UE_LOG(LogSubmitToolP4, Warning, TEXT("Revert unchanged operation cancelled."));
+			}
+			else if(Result == ECommandResult::Succeeded)
+			{
+				UE_LOG(LogSubmitToolP4, Log, TEXT("Revert unchanged operation succeeded."));
+				FilesInCL = ChangelistState->GetFilesStates();
+				FilesDepotPaths.Reset(FilesInCL.Num());
+				PrintFilesAndShelvedFiles();
+			}
+			else if(Operation->GetResultInfo().ErrorMessages.Num() == 0)
+			{
+				UE_LOG(LogSubmitToolP4, Log, TEXT("There were no unchanged files to revert."));
+			}
+			else if(Result == ECommandResult::Failed)
+			{
+				UE_LOG(LogSubmitToolP4, Error, TEXT("CL %s revert unchanged operation failed:"), *GetCLID());
+				PrintErrorMessages(Operation->GetResultInfo());
+			}
+
+			ActiveP4Operations.Remove(Operation);
+			OnRevertComplete.ExecuteIfBound(Operation, Result);
+		}
+	));
+}
+
+void FChangelistService::RehydrateDataFromP4State()
+{
+	ETaskArea ChangeType = ETaskArea::None;
+
+	if(!AreCLDescriptionsIdentical())
+	{
+		ChangeType |= ETaskArea::Changelist;
+
+		UE_LOG(LogSubmitToolP4, Log, TEXT("CL %s Description was updated outside of Submit Tool while it was still open, Description has been updated to match P4V."), *GetCLID());
+		UE_LOG(LogSubmitToolP4Debug, Log, TEXT("\n - Original Description '%s'\n - Submit Tool Description '%s'\n - New Description '%s'"), *OriginalDescription.ToString(), *CLDescription, *ChangelistState->GetDescriptionText().ToString());
+
+		OriginalDescription = ChangelistState->GetDescriptionText();
+		CLDescription = ChangelistState->GetDescriptionText().ToString();
+	}
+
+	if(ChangelistState->GetFilesStatesNum() != FilesInCL.Num())
+	{
+		ChangeType |= ETaskArea::LocalFiles;
+		FilesInCL = ChangelistState->GetFilesStates();
+		FilesDepotPaths.Reset(FilesInCL.Num());
+		UE_LOG(LogSubmitToolP4, Log, TEXT("CL %s files were updated outside of Submit Tool while it was open, Validation state has been reset"), *GetCLID());
+		PrintFilesAndShelvedFiles();
+	}
+
+	if(ChangelistState->GetShelvedFilesStatesNum() != ShelvedFilesInCL.Num())
+	{
+		ChangeType |= ETaskArea::ShelvedFiles;
+		ShelvedFilesInCL = ChangelistState->GetShelvedFilesStates();
+		ShelvedFilesDepotPaths.Reset(ShelvedFilesInCL.Num());
+		UE_LOG(LogSubmitToolP4, Log, TEXT("CL %s shelved files were updated outside of Submit Tool while it was open, Validation state has been reset"), *GetCLID());
+		PrintFilesAndShelvedFiles();
+	}
+
+	if(ChangeType != ETaskArea::None)
+	{
+		CLRefreshCallback.ExecuteIfBound(ChangeType);
+	}
+}
+
+bool FChangelistService::AreCLDescriptionsIdentical()
+{
+	const FString NewDescription = ChangelistState->GetDescriptionText()
+		.ToString()
+		.TrimStartAndEnd()
+		.TrimChar('\n');
+
+	const FString CurrentDescription = OriginalDescription
+		.ToString()
+		.TrimStartAndEnd()
+		.TrimChar('\n');
+
+	return NewDescription.Equals(CurrentDescription, ESearchCase::IgnoreCase);
+}
+
+void FChangelistService::PrintMessages(const FSourceControlResultInfo& ResultInfo) const
+{
+	for(const FText& Msg : ResultInfo.InfoMessages)
+	{
+		UE_LOG(LogSubmitToolP4Debug, Log, TEXT("%s"), *Msg.ToString());
+	}
+}
+
+void FChangelistService::PrintErrorMessages(const FSourceControlResultInfo& ResultInfo) const
+{
+	for(const FText& errorMsg : ResultInfo.ErrorMessages)
+	{
+		UE_LOG(LogSubmitToolP4, Error, TEXT("%s"), *errorMsg.ToString());
+	}
+}
+
+void FChangelistService::Submit(const FString& InDescriptionAddendum, const FSourceControlOperationComplete& OnSubmitComplete)
+{
+	UE_LOG(LogSubmitToolP4, Warning, TEXT("Submit in progress for CL: %s. Please wait..."), *GetCLID());
+
+	TSharedRef<FCheckIn, ESPMode::ThreadSafe> CheckinOp = ISourceControlOperation::Create<FCheckIn>();
+
+	FString FinalDescription = FString::Format(TEXT("{0}{1}"), { CLDescription, InDescriptionAddendum });
+
+	CheckinOp->SetDescription(FText::FromString(FinalDescription));
+
+	ActiveP4Operations.Add(CheckinOp);
+
+	CachedSCCProvider->Execute(CheckinOp, ChangelistPtr, EConcurrency::Asynchronous, FSourceControlOperationComplete::CreateLambda(
+		[this, OnSubmitComplete](const FSourceControlOperationRef& Operation, ECommandResult::Type Result) {
+			if(Result == ECommandResult::Failed)
+			{
+				UE_LOG(LogSubmitToolP4, Error, TEXT("CL %s submit operation failed:"), *GetCLID());
+				PrintErrorMessages(Operation->GetResultInfo());
+			}
+			else if(Result == ECommandResult::Succeeded)
+			{
+				TSharedRef<FCheckIn> CheckIn = StaticCastSharedRef<FCheckIn>(Operation);
+				UE_LOG(LogSubmitToolP4, Log, TEXT("Submit operation succeeded: %s"), *CheckIn->GetSuccessMessage().ToString());
+
+				FTelemetryService::Get()->SubmitSucceeded();
+			}
+			else if(Result == ECommandResult::Cancelled)
+			{
+				UE_LOG(LogSubmitToolP4, Warning, TEXT("Submit operation cancelled."));
+			}
+
+			ActiveP4Operations.Remove(Operation);
+			OnSubmitComplete.ExecuteIfBound(Operation, Result);
+		}
+	));
+}
+
+void FChangelistService::SendCLDescriptionToP4(EConcurrency::Type Concurrency, FOnP4OperationCompleteDelegate InCallback)
+{
+	if(!CLDescription.Equals(OriginalDescription.ToString()))
+	{
+		UE_LOG(LogSubmitToolP4Debug, Log, TEXT("Saving CL %s"), *GetCLID());
+		TSharedRef<FEditChangelist, ESPMode::ThreadSafe> EditCLOp = ISourceControlOperation::Create<FEditChangelist>();
+		EditCLOp->SetDescription(FText::FromString(CLDescription));
+
+		ActiveP4Operations.Add(EditCLOp);
+
+		CachedSCCProvider->Execute(EditCLOp, ChangelistPtr, Concurrency, FSourceControlOperationComplete::CreateLambda(
+			[this, InCallback](const FSourceControlOperationRef& UpdateOperation, ECommandResult::Type Result) {
+				if(Result == ECommandResult::Failed)
+				{
+					UE_LOG(LogSubmitToolP4, Error, TEXT("CL %s edit changelist operation %s failed:"), *GetCLID(), *UpdateOperation->GetName().ToString());
+					PrintErrorMessages(UpdateOperation->GetResultInfo());
+				}
+				else if(Result == ECommandResult::Succeeded)
+				{
+					UE_LOG(LogSubmitToolP4, Log, TEXT("CL %s edit changelist operation succeeded."), *GetCLID());
+					OriginalDescription = FText::FromString(CLDescription);
+				}
+				else if(Result == ECommandResult::Cancelled)
+				{
+					UE_LOG(LogSubmitToolP4, Warning, TEXT("CL %s edit changelist operation cancelled."), *GetCLID());
+				}
+
+				ActiveP4Operations.Remove(UpdateOperation);
+
+				InCallback.ExecuteIfBound();
+			}
+		));
+	}
+	else
+	{
+		InCallback.ExecuteIfBound();
+	}
+}
+
+void FChangelistService::DeleteShelvedFiles(const FSourceControlOperationComplete& OnDeleteComplete)
+{
+	UE_LOG(LogSubmitToolP4, Log, TEXT("Removing shelved files in CL %s..."), *GetCLID());
+	const TSharedRef<FDeleteShelved, ESPMode::ThreadSafe> DeleteShelvedOp = ISourceControlOperation::Create<FDeleteShelved>();
+
+	ActiveP4Operations.Add(DeleteShelvedOp);
+
+	CachedSCCProvider->Execute(DeleteShelvedOp, ChangelistPtr, EConcurrency::Asynchronous, FSourceControlOperationComplete::CreateLambda(
+		[this, OnDeleteComplete](const FSourceControlOperationRef& DeleteShelvedOp, ECommandResult::Type Result) {
+			if(Result == ECommandResult::Failed)
+			{
+				UE_LOG(LogSubmitToolP4, Error, TEXT("CL %s delete shelved operation failed:"), *GetCLID());
+
+				PrintErrorMessages(DeleteShelvedOp->GetResultInfo());
+			}
+			else if(Result == ECommandResult::Succeeded)
+			{
+				ShelvedFilesInCL.Reset();
+				ShelvedFilesDepotPaths.Reset();
+				UE_LOG(LogSubmitToolP4, Log, TEXT("CL %s delete shelved operation succeeded"), *GetCLID());
+			}
+			else if(Result == ECommandResult::Cancelled)
+			{
+				UE_LOG(LogSubmitToolP4, Warning, TEXT("CL %s delete shelved operation cancelled."), *GetCLID());
+			}
+
+			ActiveP4Operations.Remove(DeleteShelvedOp);
+			OnDeleteComplete.ExecuteIfBound(DeleteShelvedOp, Result);
+		}
+	));
+}
+
+void FChangelistService::CreateShelvedFiles(const FSourceControlOperationComplete& OnCreateComplete)
+{
+	UE_LOG(LogSubmitToolP4, Log, TEXT("Creating shelved files for CL %s..."), *GetCLID());
+	const TSharedRef<FShelve, ESPMode::ThreadSafe> CreateShelvedOp = ISourceControlOperation::Create<FShelve>();
+
+	ActiveP4Operations.Add(CreateShelvedOp);
+
+	CachedSCCProvider->Execute(CreateShelvedOp, ChangelistPtr, EConcurrency::Asynchronous, FSourceControlOperationComplete::CreateLambda(
+		[this, OnCreateComplete](const FSourceControlOperationRef& CreateShelvedOp, ECommandResult::Type Result) {
+			if(Result == ECommandResult::Failed)
+			{
+				UE_LOG(LogSubmitToolP4, Error, TEXT("CL %s create shelved operation failed:"), *GetCLID());
+
+				PrintErrorMessages(CreateShelvedOp->GetResultInfo());
+			}
+			else if(Result == ECommandResult::Succeeded)
+			{
+				ShelvedFilesInCL = ChangelistState->GetShelvedFilesStates();
+				ShelvedFilesDepotPaths.Reset(ShelvedFilesInCL.Num());
+				UE_LOG(LogSubmitToolP4, Log, TEXT("CL %s create shelved operation succeeded"), *GetCLID());
+			}
+			else if(Result == ECommandResult::Cancelled)
+			{
+				UE_LOG(LogSubmitToolP4, Warning, TEXT("CL %s create shelved operation cancelled."), *GetCLID());
+			}
+
+			ActiveP4Operations.Remove(CreateShelvedOp);
+			OnCreateComplete.ExecuteIfBound(CreateShelvedOp, Result);
+		}
+	));
+}
+
+bool FChangelistService::P4Tick(float DeltaTime)
+{
+	for(const TSharedRef<ISourceControlOperation>& operation : ActiveP4Operations)
+	{
+		UE_LOG(LogSubmitToolP4, Log, TEXT("%s operation still in progress: %s"), *operation->GetName().ToString(), *operation->GetInProgressString().ToString());
+	}
+
+	return true;
+}
+
+
+bool FChangelistService::IsP4OperationRunning(FName OperationName)
+{
+	if(OperationName == FName())
+	{
+		return ActiveP4Operations.Num() > 0;
+	}
+
+	for(const TSharedRef<ISourceControlOperation>& Operation : ActiveP4Operations)
+	{
+		if(Operation->GetName().IsEqual(OperationName, ENameCase::IgnoreCase))
+		{
+			return true;
+		}
+	}
+
+	return false;
+}
+
+void FChangelistService::CancelP4Operations(FName OperationName)
+{
+	for(const TSharedRef<ISourceControlOperation>& Operation : ActiveP4Operations)
+	{
+		if(OperationName == FName() || Operation->GetName().IsEqual(OperationName, ENameCase::IgnoreCase))
+		{
+			if(CachedSCCProvider->CanCancelOperation(Operation))
+			{
+				CachedSCCProvider->CancelOperation(Operation);
+				UE_LOG(LogSubmitToolP4, Warning, TEXT("P4 Operation %s cancelling requested"), *Operation->GetName().ToString())
+			}
+		}
+	}
+}
+
+void FChangelistService::CreateCLFromDefaultCL()
+{
+	// If we try to open submit tool with the default cl: Update all CL Status -> create a new CL -> move files from default to the new CL -> regular flow
+	UE_LOG(LogSubmitToolP4, Log, TEXT("Default changelist is not supported by Submit Tool, creating a new CL and moving files..."));
+
+	TSharedRef<FUpdatePendingChangelistsStatus> UpdatePendingChangelistsOperation = ISourceControlOperation::Create<FUpdatePendingChangelistsStatus>();
+	UpdatePendingChangelistsOperation->SetUpdateFilesStates(true);
+	UpdatePendingChangelistsOperation->SetUpdateAllChangelists(true);
+
+	ActiveP4Operations.Add(UpdatePendingChangelistsOperation);
+
+	CachedSCCProvider->Execute(UpdatePendingChangelistsOperation, EConcurrency::Asynchronous, FSourceControlOperationComplete::CreateLambda(
+		[this](const FSourceControlOperationRef& Operation, ECommandResult::Type Result) {
+			ActiveP4Operations.Remove(Operation);
+
+			if(Result == ECommandResult::Failed)
+			{
+				UE_LOG(LogSubmitToolP4, Error, TEXT("Failed to fetch changelists from p4."));
+				PrintErrorMessages(Operation->GetResultInfo());
+			}
+			else if(Result == ECommandResult::Succeeded)
+			{
+
+				TArray<FString> FilesInDefault;
+				for(const FSourceControlChangelistPtr cl : CachedSCCProvider->GetChangelists(EStateCacheUsage::Use))
+				{
+					if(cl->IsDefault())
+					{
+						for(FSourceControlStateRef file : CachedSCCProvider->GetState(cl.ToSharedRef(), EStateCacheUsage::Use)->GetFilesStates())
+						{
+							FilesInDefault.Emplace(file->GetFilename());
+						}
+					}
+				}
+
+				TSharedRef<FNewChangelist, ESPMode::ThreadSafe> NewCLOp = ISourceControlOperation::Create<FNewChangelist>();
+
+				FStringFormatNamedArguments FormatArgs = { 
+					{ TEXT("FileCount") ,FilesInDefault.Num() }
+				};
+
+				NewCLOp->SetDescription(FText::FromString(FString::Format(*Parameters.NewChangelistMessage, FormatArgs)));
+
+				ActiveP4Operations.Add(NewCLOp);
+
+				CachedSCCProvider->Execute(NewCLOp, EConcurrency::Asynchronous, FSourceControlOperationComplete::CreateLambda(
+					[this, FilesInDefault](const FSourceControlOperationRef& Operation, ECommandResult::Type Result) {
+						ActiveP4Operations.Remove(Operation);
+
+						if(Result == ECommandResult::Failed)
+						{
+							UE_LOG(LogSubmitToolP4, Error, TEXT("Failed to create new changelist from default"));
+							PrintErrorMessages(Operation->GetResultInfo());
+						}
+						else if(Result == ECommandResult::Succeeded)
+						{
+							TSharedRef<FNewChangelist> NewCLOp = StaticCastSharedRef<FNewChangelist>(Operation);
+							ChangelistPtr = NewCLOp->GetNewChangelist();
+							ChangelistState = CachedSCCProvider->GetState(ChangelistPtr.ToSharedRef(), EStateCacheUsage::Use);
+
+							OriginalDescription = ChangelistState->GetDescriptionText();
+							CLDescription = ChangelistState->GetDescriptionText().ToString();
+
+
+							TSharedRef<FMoveToChangelist> MoveOp = ISourceControlOperation::Create<FMoveToChangelist>();
+
+							ActiveP4Operations.Add(MoveOp);
+							CachedSCCProvider->Execute(MoveOp, ChangelistPtr, FilesInDefault, EConcurrency::Asynchronous, FSourceControlOperationComplete::CreateLambda(
+								[this, FileCount = FilesInDefault.Num()](const FSourceControlOperationRef& Operation, ECommandResult::Type Result) {
+									ActiveP4Operations.Remove(Operation);
+									if(Result == ECommandResult::Failed)
+									{
+										UE_LOG(LogSubmitToolP4, Error, TEXT("Move files to CL failed"));
+										PrintErrorMessages(Operation->GetResultInfo());
+									}
+									else if(Result == ECommandResult::Succeeded)
+									{
+										FilesInCL = ChangelistState->GetFilesStates();
+										FilesDepotPaths.Reset(FilesInCL.Num());
+										ShelvedFilesInCL = ChangelistState->GetShelvedFilesStates();
+										ShelvedFilesDepotPaths.Reset(ShelvedFilesInCL.Num());
+										CLID = ChangelistPtr->GetIdentifier();
+										FConfiguration::AddOrUpdateEntry(TEXT("$(CL)"), GetCLID());
+
+										UE_LOG(LogSubmitToolP4, Log, TEXT("Created CL %s and moved with %d files from the default CL."), *ChangelistPtr->GetIdentifier(), FileCount);
+										UE_LOG(LogSubmitToolDebug, Log, TEXT("ChangeListService CLID Updated to %s"), *GetCLID());
+										UE_LOG(LogSubmitToolDebug, Log, TEXT("Configuration updated value $(CL) to %s"), *GetCLID());
+
+										CLReadyCallback.ExecuteIfBound(ChangelistState.IsValid());
+									}
+									else if(Result == ECommandResult::Cancelled)
+									{
+										UE_LOG(LogSubmitToolP4, Warning, TEXT("Move files to CL cancelled."), *GetCLID());
+									}
+								}));
+						}
+						else if(Result == ECommandResult::Cancelled)
+						{
+							UE_LOG(LogSubmitToolP4, Warning, TEXT("Create new CL was cancelled."), *GetCLID());
+						}
+					}
+				));
+			}
+			else if(Result == ECommandResult::Cancelled)
+			{
+				UE_LOG(LogSubmitToolP4, Warning, TEXT("CL %s edit changelist operation cancelled."), *GetCLID());
+			}
+		}));
+}
+
+TArray<FSourceControlChangelistStatePtr> FChangelistService::GetOtherChangelistsStates()
+{
+	TArray<FSourceControlChangelistStatePtr> Output;
+
+	for(const FSourceControlChangelistRef& CL : CachedSCCProvider->GetChangelists(EStateCacheUsage::Use))
+	{
+		if(CL->GetIdentifier() != CLID)
+		{
+			FSourceControlChangelistPtr ClPtr = CL;
+			Output.Add(CachedSCCProvider->GetState(ClPtr.ToSharedRef(), EStateCacheUsage::Use));
+		}
+	}
+
+	return Output;
+}
+
+void FChangelistService::PrintFilesAndShelvedFiles()
+{
+	if(FilesInCL.Num() > 0)
+	{
+		UE_LOG(LogSubmitToolP4Debug, Log, TEXT("Files in CL:"));
+		for(const FSourceControlStateRef& File : FilesInCL)
+		{
+			UE_LOG(LogSubmitToolP4Debug, Log, TEXT("\n%s"), *File->GetFilename());
+		}
+	}
+
+	if(ShelvedFilesInCL.Num() > 0)
+	{
+		UE_LOG(LogSubmitToolP4Debug, Log, TEXT("Shelved Files in CL:"));
+		for(const FSourceControlStateRef& File : ShelvedFilesInCL)
+		{
+			UE_LOG(LogSubmitToolP4Debug, Log, TEXT("\n%s"), *File->GetFilename());
+		}
+	}
+}
+
+const TArray<FString>& FChangelistService::GetFilesDepotPaths(bool bForce)
+{
+	FScopeLock Lock(&Mutex);
+	if(bForce || (FilesDepotPaths.IsEmpty() && !FilesInCL.IsEmpty()))
+	{
+		UE::Tasks::TTask<bool> OpenedTask = SourceControlService->RunCommand(TEXT("opened"), { "-c", CLID }, FOnSCCCommandComplete::CreateLambda(
+			[this](bool bSuccess, const FSCCRecordSet& InResultValues, const FSourceControlResultInfo& InResultsInfo)
+			{
+				if(bSuccess)
+				{
+					for(const TMap<FString, FString>& Record : InResultValues)
+					{
+						if(Record.Contains(TEXT("depotFile")))
+						{
+							FilesDepotPaths.Add(Record[TEXT("depotFile")]);
+						}
+						else
+						{
+							const FString Base = TEXT("depotFile");
+							size_t i = 0;
+							FString PathsKey = FString::Printf(TEXT("%s%d"), *Base, i);
+							while(Record.Contains(PathsKey))
+							{
+								FilesDepotPaths.Add(Record[PathsKey]);
+								++i;
+								PathsKey = FString::Printf(TEXT("%s%d"), *Base, i);
+							}
+						}
+					}
+				}
+			}));
+		OpenedTask.Wait();
+	}
+
+	return FilesDepotPaths;
+}
+
+const TArray<FString>& FChangelistService::GetShelvedFilesDepotPaths(bool bForce)
+{
+	// Don't bother with the default CL
+	if (CLID.Equals("default", ESearchCase::IgnoreCase))
+	{
+		return ShelvedFilesDepotPaths;
+	}
+
+	FScopeLock Lock(&Mutex);
+	if(bForce || (ShelvedFilesDepotPaths.IsEmpty() && !ShelvedFilesInCL.IsEmpty()))
+	{
+		UE::Tasks::TTask<bool> DescribeTask = SourceControlService->RunCommand(TEXT("describe"), { "-S", CLID }, FOnSCCCommandComplete::CreateLambda(
+			[this](bool bSuccess, const FSCCRecordSet& InResultValues, const FSourceControlResultInfo& InResultsInfo)
+			{
+				if(bSuccess)
+				{
+					for(const TMap<FString, FString>& Record : InResultValues)
+					{
+						if(Record.Contains(TEXT("depotFile")))
+						{
+							FilesDepotPaths.Add(Record[TEXT("depotFile")]);
+						}
+						else
+						{
+							const FString Base = TEXT("depotFile");
+							size_t i = 0;
+							FString PathsKey = FString::Printf(TEXT("%s%d"), *Base, i);
+							while(Record.Contains(PathsKey))
+							{
+								ShelvedFilesDepotPaths.Add(Record[PathsKey]);
+								++i;
+								PathsKey = FString::Printf(TEXT("%s%d"), *Base, i);
+							}
+						}
+					}
+				}
+			}));
+		DescribeTask.Wait();
+	}
+
+	return ShelvedFilesDepotPaths;
+}
diff --git a/Engine/Source/Programs/SubmitTool/Private/Logic/ChangelistService.h b/Engine/Source/Programs/SubmitTool/Private/Logic/ChangelistService.h
new file mode 100644
--- /dev/null
+++ b/Engine/Source/Programs/SubmitTool/Private/Logic/ChangelistService.h
@@ -0,0 +1,148 @@
+// Copyright Epic Games, Inc. All Rights Reserved.
+
+#pragma once
+
+#include "CoreMinimal.h"
+#include "ISourceControlProvider.h"
+#include "ISourceControlChangelist.h"
+#include "ISourceControlChangelistState.h"
+#include "Containers/Ticker.h"
+#include "Parameters/SubmitToolParameters.h"
+#include "Services/Interfaces/ISubmitToolService.h"
+
+struct FSourceControlResultInfo;
+class ISTSourceControlService;
+using FSourceControlOperationRef = TSharedRef<class ISourceControlOperation, ESPMode::ThreadSafe>;
+
+
+DECLARE_DELEGATE(FOnP4OperationCompleteDelegate)
+DECLARE_DELEGATE_OneParam(FOnChangeListReadyDelegate, bool /*bValidState*/)
+DECLARE_DELEGATE_OneParam(FOnChangelistRefreshDelegate, ETaskArea /*ChangedValue*/)
+DECLARE_MULTICAST_DELEGATE(FOnCLDescriptionUpdated)
+
+
+class FChangelistService final : public ISubmitToolService
+{
+public:
+
+	FChangelistService(const FGeneralParameters& InParameters, const TSharedPtr<ISTSourceControlService> SourceControlService, const FOnChangeListReadyDelegate& InCLReadyCallback, const FOnChangelistRefreshDelegate& InCLRefreshCallback);
+	~FChangelistService();
+
+	const FString GetCLID() const
+	{
+		if(ChangelistPtr.IsValid())
+		{
+			return ChangelistPtr->GetIdentifier();
+		} 
+		
+		if(!CLID.IsEmpty())
+		{
+			return CLID;
+		}
+
+		return TEXT("Invalid");
+	}
+
+	const FString& GetCLDescription()
+	{
+		return CLDescription;
+	}
+
+	bool SetCLDescription(const FString& newDescription, bool bNotifyEvent = false)
+	{
+		FString lineEndReplaced = newDescription.Replace(TEXT("\r\n"), TEXT("\n"));
+
+		if(CLDescription.Equals(lineEndReplaced, ESearchCase::IgnoreCase))
+		{
+			return false;
+		}
+
+		CLDescription = lineEndReplaced;
+
+		if(bNotifyEvent)
+		{
+			OnCLDescriptionUpdated.Broadcast();
+		}
+
+		return true;
+	}
+
+	const TArray<FSourceControlStateRef>& GetFilesInCL() const
+	{
+		return FilesInCL;
+	}
+
+	const TArray<FSourceControlStateRef>& GetShelvedFilesInCL() const
+	{
+		return ShelvedFilesInCL;
+	}
+
+	bool HasShelvedFiles() const
+	{
+		return ShelvedFilesInCL.Num() != 0;
+	}
+
+	bool HasP4OperationsRunning() const
+	{
+		return ActiveP4Operations.Num() > 0;
+	}
+
+
+	void Init();
+	void Submit(const FString& InDescriptionAddendum = TEXT(""), const FSourceControlOperationComplete& OnSubmitComplete = FSourceControlOperationComplete());
+	void CreateCLFromDefaultCL();
+	void FetchChangelistDataAsync();
+	void RevertUnchangedFilesAsync(const FSourceControlOperationComplete& OnRevertComplete = FSourceControlOperationComplete());
+	void SendCLDescriptionToP4(EConcurrency::Type Concurrency = EConcurrency::Asynchronous, FOnP4OperationCompleteDelegate InCallback = nullptr);
+	void DeleteShelvedFiles(const FSourceControlOperationComplete& OnDeleteComplete);
+	void CreateShelvedFiles(const FSourceControlOperationComplete& OnCreateComplete);
+	bool P4Tick(float DeltaTime);
+
+	bool IsP4OperationRunning(FName OperationName = FName());
+	void CancelP4Operations(FName OperationName = FName());
+
+	FOnCLDescriptionUpdated OnCLDescriptionUpdated;
+
+	TArray<FSourceControlChangelistStatePtr> GetOtherChangelistsStates();
+
+	const TArray<FString>& GetFilesDepotPaths(bool bForce = false);
+	const TArray<FString>& GetShelvedFilesDepotPaths(bool bForce = false);
+
+private:
+	mutable FCriticalSection Mutex;
+
+	FString CLID;
+	FString CurrentStream;
+	const FGeneralParameters& Parameters;
+	const FOnChangeListReadyDelegate& CLReadyCallback;
+	const FOnChangelistRefreshDelegate& CLRefreshCallback;
+
+	const FTSTicker::FDelegateHandle TickHandle;
+
+	TArray<TSharedRef<ISourceControlOperation>> ActiveP4Operations;
+
+	FSourceControlChangelistPtr ChangelistPtr = nullptr;
+	TArray<FSourceControlStateRef> FilesInCL;
+	TArray<FSourceControlStateRef> ShelvedFilesInCL;
+	TArray<FString> FilesDepotPaths;
+	TArray<FString> ShelvedFilesDepotPaths;
+	FSourceControlChangelistStatePtr ChangelistState = nullptr;
+	FText OriginalDescription;
+	FString CLDescription;
+	ISourceControlProvider* CachedSCCProvider;
+	const TSharedPtr<ISTSourceControlService> SourceControlService;
+
+	void FindInitialChangelistsAsync();
+
+	void FetchChangelistCallback(const FSourceControlOperationRef& UpdateOperation, ECommandResult::Type Result);
+	void RehydrateDataFromP4State();
+	bool AreCLDescriptionsIdentical();
+	void PrintFilesAndShelvedFiles();
+
+	void PrintMessages(const FSourceControlResultInfo& ResultInfo) const;
+	void PrintErrorMessages(const FSourceControlResultInfo& ResultInfo) const;
+
+
+};
+
+Expose_TNameOf(FChangelistService);
\ No newline at end of file
diff --git a/Engine/Source/Programs/SubmitTool/Private/Logic/CredentialsService.cpp b/Engine/Source/Programs/SubmitTool/Private/Logic/CredentialsService.cpp
new file mode 100644
--- /dev/null
+++ b/Engine/Source/Programs/SubmitTool/Private/Logic/CredentialsService.cpp
@@ -0,0 +1,373 @@
+// Copyright Epic Games, Inc. All Rights Reserved.
+
+#include "CredentialsService.h"
+
+#include "SubmitToolUtils.h"
+#include "Logging/LogMacros.h"
+#include "Logging/SubmitToolLog.h"
+#include "Logic/ProcessWrapper.h"
+
+#include "JsonObjectConverter.h"
+#include "HAL/FileManager.h"
+#include "Misc/FileHelper.h"
+#include "Misc/Paths.h"
+#include "Misc/Base64.h"
+#include "Tasks/Task.h"
+
+TUniquePtr<FAES::FAESKey> FCredentialsService::Key = nullptr;
+
+FCredentialsService::FCredentialsService(const FOAuthTokenParams& InOAuthParameters)
+	: Parameters(InOAuthParameters)
+{
+	if(IsOIDCTokenEnabled())
+	{
+		GetOIDCToken();
+		FTSTicker::GetCoreTicker().AddTicker(FTickerDelegate::CreateRaw(this, &FCredentialsService::Tick), 5);
+	}
+
+	LoadKey();
+	LoadCredentials();
+}
+
+void FCredentialsService::LoadKey()
+{
+	if(!IFileManager::Get().FileExists(*GetKeyFilepath()))
+	{
+		GenerateKey();
+	}
+
+	if(IFileManager::Get().FileExists(*GetKeyFilepath()))
+	{
+		FArchive* File = IFileManager::Get().CreateFileReader(*GetKeyFilepath());
+		if(File->TotalSize() < 4)
+		{
+			UE_LOG(LogSubmitToolDebug, Warning, TEXT("Unexpected file size encryption key invalidated"));
+			File->Close();
+			delete File;
+			File = nullptr;
+			return;
+		}
+
+		int32 Size;
+		*File << Size;
+
+		// see if the file has exactly the length we expect
+		// two int32 (a size and a garbage one) + the data + one garbage int32 in between the data
+		if(File->TotalSize() != 4 + 4 + Size + 4)
+		{
+			UE_LOG(LogSubmitToolDebug, Warning, TEXT("Unexpected file size encryption key invalidated"));
+			File->Close();
+			delete File;
+			File = nullptr;
+			return;
+		}
+
+		int32 Garbage;
+		*File << Garbage;
+
+		TArray<uint8> Bytes;
+		uint8 byte;
+
+		for(size_t i = 0; i < Size; ++i)
+		{
+			if(i == 2)
+			{
+				*File << Garbage;
+			}
+
+			*File << byte;
+			Bytes.Add(byte);
+		}
+		File->Close();
+		delete File;
+		File = nullptr;
+
+		Key = MakeUnique<FAES::FAESKey>();
+		check(Bytes.Num() == sizeof(FAES::FAESKey::Key));
+		FMemory::Memcpy(Key->Key, &Bytes[0], sizeof(FAES::FAESKey::Key));
+	}
+}
+
+void FCredentialsService::GenerateKey()
+{
+	TArray<uint8> dataArray;
+	for(size_t i = 0; i < FAES::FAESKey::KeySize; ++i)
+	{
+		int32 random = FMath::Rand();
+		dataArray.Add(random);
+	}
+
+	Key = MakeUnique<FAES::FAESKey>();
+	FMemory::Memcpy(Key->Key, dataArray.GetData(), sizeof(FAES::FAESKey::Key));
+
+	FArchive* File = IFileManager::Get().CreateFileWriter(*GetKeyFilepath(), EFileWrite::FILEWRITE_EvenIfReadOnly);
+	int32 Size = FAES::FAESKey::KeySize;
+	*File << Size;
+	int32 Garbage = FMath::Rand();
+	*File << Garbage;
+
+	for(size_t i = 0; i < Size; ++i)
+	{
+		if(i == 2)
+		{
+			Garbage = FMath::Rand();
+			*File << Garbage;
+		}
+
+		*File << dataArray[i];
+	}
+
+	File->Close();
+	delete File;
+	File = nullptr;
+}
+
+const FString FCredentialsService::GetKeyFilepath()
+{
+	return FPaths::Combine(FSubmitToolUtils::GetLocalAppDataPath(), TEXT("SubmitTool"), TEXT(".cache"));
+}
+
+const FString FCredentialsService::GetCredentialsFilepath() const
+{
+	return FPaths::Combine(FSubmitToolUtils::GetLocalAppDataPath(), TEXT("SubmitTool"), TEXT("jira.dat"));
+}
+
+void FCredentialsService::GetOIDCToken()
+{	
+	UE_LOG(LogSubmitTool, Log, TEXT("Obtaining new OIDCToken"));
+	GetOIDCTask = UE::Tasks::Launch(UE_SOURCE_LOCATION,
+		[this] {
+			FOnOutputLine OutputLineProcess = FOnOutputLine::CreateLambda([this](const FString& OutputLine, const EProcessOutputType& OutputType) {
+				if(OutputType == EProcessOutputType::ProcessError)
+				{
+					UE_LOG(LogSubmitTool, Error, TEXT("%s"), *OutputLine);
+				}
+				else
+				{
+					// OIDCToken console output token starts with a {
+					if(OutputLine.StartsWith(TEXT("{")))
+					{
+						ParseOIDCTokenData(OutputLine);
+					}
+					else
+					{
+						UE_LOG(LogSubmitToolDebug, Log, TEXT("%s"), *OutputLine);
+					}
+				}
+				});
+
+			FProcessWrapper GetToken = FProcessWrapper(TEXT("Oidc"), Parameters.OAuthTokenTool, FString::Printf(TEXT("%s"), *Parameters.OAuthArgs), FOnCompleted(), OutputLineProcess);
+			GetToken.Start(true);
+
+			if(GetToken.ExitCode != 0)
+			{
+				UE_LOG(LogSubmitTool, Warning, TEXT("Couldn't obtain OIDC credentials"));
+				return false;
+			}
+
+			return true;
+		});
+}
+
+bool FCredentialsService::ParseOIDCTokenData(const FString& InToken)
+{
+	TSharedPtr<FJsonObject> RootJsonObject;
+	TSharedRef<TJsonReader<>> Reader = TJsonReaderFactory<>::Create(InToken);
+	FJsonSerializer::Deserialize(Reader, RootJsonObject);
+
+	if(RootJsonObject.IsValid())
+	{
+		FString Expiration = RootJsonObject->GetStringField(TEXT("ExpiresAt"));
+		FDateTime::ParseIso8601(*Expiration, TokenExpiration);
+		UE_LOG(LogSubmitToolDebug, Log, TEXT("OIDC Token Expiration %s"), *Expiration);
+		
+		OIDCToken = RootJsonObject->GetStringField(TEXT("Token"));
+		UE_LOG(LogSubmitTool, Log, TEXT("OIDC Token loaded correctly"));
+		return true;
+	}
+
+	return false;
+}
+
+constexpr int JiraCredentialDatVersion = 1;
+void FCredentialsService::SaveCredentials() const
+{
+	TArray<uint8> Bytes;
+	Bytes.SetNumUninitialized(LoginString.Len());
+	StringToBytes(LoginString, Bytes.GetData(), LoginString.Len());
+
+	int32 ActualLength = Bytes.Num();
+
+	int32 NumBytesEncrypted = Align(Bytes.Num(), FAES::AESBlockSize);
+	Bytes.SetNum(NumBytesEncrypted);
+	FAES::EncryptData(Bytes.GetData(), Bytes.Num(), *Key);
+
+	FArchive* File = IFileManager::Get().CreateFileWriter(*GetCredentialsFilepath(), EFileWrite::FILEWRITE_EvenIfReadOnly);
+
+	if(File != nullptr)
+	{
+		int32 Version = JiraCredentialDatVersion;
+		*File << Version;
+
+		*File << NumBytesEncrypted;
+		*File << ActualLength;
+		int32 Garbage = FMath::Rand();
+		*File << Garbage;
+		File->Serialize(Bytes.GetData(), Bytes.Num());
+		Garbage = FMath::Rand();
+		*File << Garbage;
+
+		File->Close();
+		delete File;
+		File = nullptr;
+	}
+	else
+	{
+		UE_LOG(LogSubmitTool, Warning, TEXT("Could not create file '%s'."), *GetCredentialsFilepath());
+	}
+}
+
+void FCredentialsService::LoadCredentials()
+{
+	if(Key.IsValid())
+	{
+		if(IFileManager::Get().FileExists(*GetCredentialsFilepath()))
+		{
+			FArchive* File = IFileManager::Get().CreateFileReader(*GetCredentialsFilepath());
+
+			if(File != nullptr)
+			{
+				// Read the size
+				if(File->TotalSize() < 4 + 4)
+				{
+					UE_LOG(LogSubmitToolDebug, Warning, TEXT("Unexpected file size login key invalidated"));
+					File->Close();
+					delete File;
+					return;
+				}
+
+				int32 Version;
+				*File << Version;
+
+				// Check Versions here
+				if(Version != JiraCredentialDatVersion)
+				{
+					UE_LOG(LogSubmitToolDebug, Warning, TEXT("Unexpected Credentials Version, aborting credentials loading."));
+					File->Close();
+					delete File;
+					return;
+				}
+
+				int32 PaddedLength;
+				int32 LengthWithoutPadding;
+
+				*File << PaddedLength;
+				*File << LengthWithoutPadding;
+
+				// see if the file has exactly the length we expect
+				// four int32 (Version, two sizes and one garbage) + the data + a final garbage int32
+				if(File->TotalSize() != (4 * sizeof(int32) + PaddedLength + sizeof(int32)))
+				{
+					UE_LOG(LogSubmitToolDebug, Warning, TEXT("Unexpected file size login invalidated"));
+					File->Close();
+					delete File;
+					return;
+				}
+
+				int32 Garbage;
+				*File << Garbage;
+
+				TArray<uint8> DeserializedBytes;
+				DeserializedBytes.SetNum(PaddedLength);
+				File->Serialize(DeserializedBytes.GetData(), PaddedLength);
+
+				FAES::DecryptData(DeserializedBytes.GetData(), DeserializedBytes.Num(), *Key);
+
+				LoginString = BytesToString(DeserializedBytes.GetData(), LengthWithoutPadding);
+
+				if(!GetUsername().IsEmpty() && !GetPassword().IsEmpty())
+				{
+					UE_LOG(LogSubmitTool, Log, TEXT("Local Credentials loaded"));
+				}
+
+				File->Close();
+				delete File;
+			}
+			else
+			{
+				UE_LOG(LogSubmitTool, Warning, TEXT("Could not read file '%s'."), *GetCredentialsFilepath());
+			}
+		}
+		else
+		{
+			UE_LOG(LogSubmitToolDebug, Warning, TEXT("File %s does not exists, no credentials were loaded"), *GetCredentialsFilepath())
+		}
+	}
+}
+
+FString FCredentialsService::GetUsername() const
+{
+	FString DecodedString;
+	if(!FBase64::Decode(LoginString, DecodedString))
+	{
+		UE_LOG(LogSubmitToolDebug, Error, TEXT("Error while trying to decode Jira Login"));
+	}
+
+	TArray<FString> LoginValues;
+	DecodedString.ParseIntoArray(LoginValues, TEXT(":"));
+
+	if(LoginValues.Num() == 2)
+	{
+		return LoginValues[0];
+	}
+
+	return FString();
+}
+
+FString FCredentialsService::GetPassword() const
+{
+	FString DecodedString;
+	if(!FBase64::Decode(LoginString, DecodedString))
+	{
+		UE_LOG(LogSubmitToolDebug, Error, TEXT("Error while trying to decode Jira Password"));
+	}
+
+	TArray<FString> LoginValues;
+	DecodedString.ParseIntoArray(LoginValues, TEXT(":"));
+
+	if(LoginValues.Num() == 2)
+	{
+		return LoginValues[1];
+	}
+
+	return FString();
+}
+
+void FCredentialsService::SetLogin(const FString& InUsername, const FString& InPassword)
+{
+	int32 ChopLocation = InUsername.Find(TEXT("@"));
+	FString FormattedUsername = InUsername;
+
+	// Just grab the username if they accidentally entered their full email.
+	if(ChopLocation != -1)
+	{
+		FormattedUsername = InUsername.LeftChop(InUsername.Len() - ChopLocation);
+	}
+
+	FString newLogin = FBase64::Encode(FormattedUsername + TEXT(":") + InPassword);
+	if(newLogin != LoginString)
+	{
+		LoginString = newLogin;
+		this->SaveCredentials();
+	}
+}
+
+bool FCredentialsService::Tick(float DeltaTime)
+{
+	if(TokenExpiration != FDateTime() && TokenExpiration < FDateTime::UtcNow())
+	{
+		GetOIDCToken();
+	}
+
+	return true;
+}
diff --git a/Engine/Source/Programs/SubmitTool/Private/Logic/CredentialsService.h b/Engine/Source/Programs/SubmitTool/Private/Logic/CredentialsService.h
new file mode 100644
--- /dev/null
+++ b/Engine/Source/Programs/SubmitTool/Private/Logic/CredentialsService.h
@@ -0,0 +1,95 @@
+// Copyright Epic Games, Inc. All Rights Reserved.
+
+#pragma once
+#include "Misc/AES.h"
+#include "Parameters/SubmitToolParameters.h"
+#include "Tasks/Task.h"
+#include "Services/Interfaces/ISubmitToolService.h"
+
+class FCredentialsService : public ISubmitToolService
+{
+public:
+	FCredentialsService(const FOAuthTokenParams& InOAuthParameters);
+	~FCredentialsService() = default;
+
+	bool HasCredentials() const
+	{
+		return !LoginString.IsEmpty();
+	}
+
+	bool AreCredentialsValid() const
+	{
+		return bValidatedCredentials;
+	}
+
+	void SetCredentialsValid(bool bValid)
+	{
+		bValidatedCredentials = bValid;
+	}
+
+	const FString& GetEncodedLoginString() const
+	{
+		return LoginString;
+	}
+
+	FString GetUsername() const;
+	void SetLogin(const FString& InUsername, const FString& InPassword);
+
+	bool IsOIDCTokenEnabled()
+	{
+		return !Parameters.OAuthTokenTool.IsEmpty();
+	}
+
+	bool IsTokenReady() const
+	{
+		return !OIDCToken.IsEmpty();
+	}
+
+	const FString& GetToken() const
+	{
+		return OIDCToken;
+	}
+
+	static const TUniquePtr<FAES::FAESKey>& GetEncryptionKey()
+	{
+		if(!Key.IsValid())
+		{
+			LoadKey();
+		}
+
+		return Key;
+	}
+
+	UE::Tasks::TTask<void> QueueWorkForToken(TFunction<void(const FString&)> InFunction)
+	{
+		return UE::Tasks::Launch(UE_SOURCE_LOCATION, [this, InFunction]() {
+			InFunction(OIDCToken);
+			}, GetOIDCTask);
+	}
+
+private:
+	static TUniquePtr<FAES::FAESKey> Key;
+	static void LoadKey();
+	static void GenerateKey();
+	static const FString GetKeyFilepath();
+
+	UE::Tasks::TTask<bool> GetOIDCTask;
+	FString OIDCToken;
+	FDateTime TokenExpiration;
+	const FOAuthTokenParams Parameters;
+
+	void GetOIDCToken();
+	bool ParseOIDCTokenData(const FString& InToken);
+	FString GetPassword() const;
+	void SaveCredentials() const;
+	void LoadCredentials();
+
+	bool Tick(float DeltaTime);
+
+	const FString GetCredentialsFilepath() const;
+
+	FString LoginString;
+	bool bValidatedCredentials = true;
+};
+
+Expose_TNameOf(FCredentialsService);
\ No newline at end of file
diff --git a/Engine/Source/Programs/SubmitTool/Private/Logic/DialogFactory.cpp b/Engine/Source/Programs/SubmitTool/Private/Logic/DialogFactory.cpp
new file mode 100644
--- /dev/null
+++ b/Engine/Source/Programs/SubmitTool/Private/Logic/DialogFactory.cpp
@@ -0,0 +1,51 @@
+// Copyright Epic Games, Inc. All Rights Reserved.
+
+#include "DialogFactory.h"
+#include "View/Widgets/ConfirmDialogWidget.h"
+#include "Widgets/SWindow.h"
+#include "Framework/Application/SlateApplication.h"
+
+EDialogFactoryResult FDialogFactory::ShowDialog(const FText& Title, const FText& Description,
+	const TArray<FString>& AvailableButtons)
+{	
+	TSharedPtr<SWindow> Window = SNew(SWindow)
+	.Title(Title)
+	.SizingRule(ESizingRule::Autosized)
+	.SupportsMaximize(false)
+	.SupportsMinimize(false)
+	.MinWidth(400)
+	.MinHeight(1);
+	
+	EDialogFactoryResult ButtonClicked = EDialogFactoryResult::ClosedWithX;
+	const FOnResult ResultsCallback = FOnResult::CreateLambda([&Window, &ButtonClicked](size_t InResult)
+	{
+		ButtonClicked = static_cast<EDialogFactoryResult>(InResult);
+		Window->RequestDestroyWindow();
+	});
+	
+	const TSharedPtr<SConfirmDialogWidget> ConfirmWidget = SNew(SConfirmDialogWidget)
+		.DescriptionText(Description)
+		.Buttons(AvailableButtons)
+		.ResultCallback(ResultsCallback);
+
+	Window->SetContent(ConfirmWidget.ToSharedRef());
+
+	FSlateApplication::Get().AddModalWindow(Window.ToSharedRef(), nullptr);
+
+	return ButtonClicked;
+}
+
+EDialogFactoryResult FDialogFactory::ShowConfirmDialog(const FText& Title, const FText& Description)
+{
+	return ShowDialog(Title, Description, TArray<FString>{ TEXT("Confirm"), TEXT("Cancel") });
+};
+
+EDialogFactoryResult FDialogFactory::ShowInformationDialog(const FText& Title, const FText& Description)
+{
+	return ShowDialog(Title, Description, TArray<FString>{ TEXT("Ok") });
+}
+
+EDialogFactoryResult FDialogFactory::ShowYesNoDialog(const FText& Title, const FText& Description)
+{
+	return ShowDialog(Title, Description, TArray<FString>{ TEXT("Yes"), TEXT("No") });
+}
diff --git a/Engine/Source/Programs/SubmitTool/Private/Logic/DialogFactory.h b/Engine/Source/Programs/SubmitTool/Private/Logic/DialogFactory.h
new file mode 100644
--- /dev/null
+++ b/Engine/Source/Programs/SubmitTool/Private/Logic/DialogFactory.h
@@ -0,0 +1,48 @@
+// Copyright Epic Games, Inc. All Rights Reserved.
+
+#pragma once
+
+#include "View/Widgets/ConfirmDialogWidget.h"
+
+enum class EDialogFactoryResult : size_t
+{
+	FirstButton = 0,
+	SecondButton,
+	ThirdButton,
+	FourthButton,
+	ClosedWithX = std::numeric_limits<size_t>::max(),
+	Confirm = FirstButton,
+	Cancel = SecondButton,
+	Yes = FirstButton,
+	No = SecondButton
+};
+
+class FDialogFactory
+{
+public:
+	static EDialogFactoryResult ShowDialog(const FText& Title, const FText& Description, const TArray<FString>& AvailableButtons);
+
+	/**
+	 * @brief Shows a dialog screen with the Cancel and Confirm buttons
+	 * @param Title The title of the dialog
+	 * @param Description The description of the dialog
+	 * @return The button idx that is pressed
+	 */
+	static EDialogFactoryResult ShowConfirmDialog(const FText& Title, const FText& Description);
+
+	/**
+	 * @brief Shows a dialog screen with the Confirm button
+	 * @param Title The title of the dialog
+	 * @param Description The description of the dialog
+	 * @return The button idx that is pressed
+	 */
+	static EDialogFactoryResult ShowInformationDialog(const FText& Title, const FText& Description);
+
+	/**
+	 * @brief Shows a dialog screen with the Yes and No buttons
+	 * @param Title The title of the dialog
+	 * @param Description The description of the dialog
+	 * @return The button idx that is pressed
+	 */
+	static EDialogFactoryResult ShowYesNoDialog(const FText& Title, const FText& Description);
+};
diff --git a/Engine/Source/Programs/SubmitTool/Private/Logic/IntegrationService.cpp b/Engine/Source/Programs/SubmitTool/Private/Logic/IntegrationService.cpp
new file mode 100644
--- /dev/null
+++ b/Engine/Source/Programs/SubmitTool/Private/Logic/IntegrationService.cpp
@@ -0,0 +1,344 @@
+// Copyright Epic Games, Inc. All Rights Reserved.
+
+#include "IntegrationService.h"
+
+#include "HAL/PlatformProcess.h"
+#include "Misc/Paths.h"
+
+#include "CommandLine/CmdLineParameters.h"
+#include "Logging/SubmitToolLog.h"
+#include "Logic/ChangelistService.h"
+#include "Logic/TasksService.h"
+#include "Logic/TagService.h"
+#include "Logic/JiraService.h"
+#include "Logic/SwarmService.h"
+#include "Logic/Services/Interfaces/ICacheDataService.h"
+#include "Logic/Services/Interfaces/ISTSourceControlService.h"
+#include "Logic/DialogFactory.h"
+#include "Internationalization/Regex.h"
+
+#include "Models/IntegrationOptions.h"
+
+#include "Configuration/Configuration.h"
+
+FIntegrationService::FIntegrationService
+(const FIntegrationParameters& InParameters,
+	TWeakPtr<FSubmitToolServiceProvider> InServiceProvider) :
+	Parameters(InParameters),
+	ServiceProvider(InServiceProvider)
+{
+	const FString CLID = ServiceProvider.Pin()->GetService<FChangelistService>()->GetCLID();
+	for(FJiraIntegrationField& Field : Parameters.Fields)
+	{
+		if(Field.Name.IsEmpty())
+		{
+			Field.Name = FGuid::NewGuid().ToString(EGuidFormats::Digits);
+		}
+
+		if(IntegrationOptions.Contains(Field.Name))
+		{
+			UE_LOG(LogSubmitTool, Error, TEXT("Integration UI has more than one Name value for item %s, Name property must be unique."), *Field.Name);
+		}
+		TSharedPtr<FIntegrationOptionBase> NewField;
+		switch(Field.Type)
+		{
+		case EFieldType::Bool:
+			NewField = IntegrationOptions.Add(Field.Name, MakeShared<FIntegrationBoolOption>(Field));
+			break;
+
+		case EFieldType::Text:
+		case EFieldType::MultiText:
+		case EFieldType::PerforceUser:
+			NewField = IntegrationOptions.Add(Field.Name, MakeShared<FIntegrationTextOption>(Field));
+			break;
+		case EFieldType::Combo:
+			NewField = IntegrationOptions.Add(Field.Name, MakeShared<FIntegrationComboOption>(Field));
+			break;
+
+		case EFieldType::UILabel:
+		case EFieldType::UISpace:
+			NewField = IntegrationOptions.Add(Field.Name, MakeShared<FIntegrationEmptyOption>(Field));
+			break;
+
+		default:
+			UE_LOG(LogSubmitTool, Warning, TEXT("Unknown integration type %s"), *StaticEnum<EFieldType>()->GetNameStringByValue(static_cast<int64>(Field.Type)));
+		}
+
+		if(NewField.IsValid() && !CLID.Equals(TEXT("default")))
+		{
+			const FString& CachedValue = ServiceProvider.Pin()->GetService<ICacheDataService>()->GetIntegrationFieldValue(CLID, Field.Name);
+
+			if(!CachedValue.IsEmpty())
+			{
+				NewField->SetValue(CachedValue);
+			}
+		}
+
+	}
+}
+
+bool FIntegrationService::OpenIntegrationTool() const
+{
+	FString Path = FConfiguration::SubstituteAndNormalizeFilename(Parameters.Path);
+	FString Args = FConfiguration::Substitute(Parameters.Args);
+
+	if(FPaths::FileExists(Path))
+	{
+
+		FProcHandle Handle = FPlatformProcess::CreateProc(*Path, *Args, false, false, false, nullptr, 0, nullptr, nullptr, nullptr);
+		if(Handle.IsValid())
+		{
+			UE_LOG(LogSubmitTool, Log, TEXT("[FNIntegrationService] Successfully invoked the FN Integration Tool (%s) with parameters (%s)."), *Path, *Args);
+		}
+		else
+		{
+			UE_LOG(LogSubmitTool, Warning, TEXT("[FNIntegrationService] Failed to invoke the FN Integration Tool (%s) with parameters (%s)."), *Path, *Args);
+		}
+
+		return Handle.IsValid();
+	}
+	else
+	{
+		UE_LOG(LogSubmitTool, Error, TEXT("[FNIntegrationService] Could not invoke the FN Integration tool, (%s) does not exist."), *Path);
+	}
+
+	return false;
+}
+
+void FIntegrationService::RequestIntegration(const FOnBooleanValueChanged OnComplete) const
+{
+	TSharedPtr<FSwarmService> SwarmService = ServiceProvider.Pin()->GetService<FSwarmService>();
+	SwarmService->FetchReview(OnGetReviewComplete::CreateLambda(
+	[this, SwarmService, OnComplete](const TUniquePtr<FSwarmReview>& InReview, const FString& InErrorMessage)
+	{
+		TSharedPtr<FChangelistService> ChangelistService = ServiceProvider.Pin()->GetService<FChangelistService>();
+		TSharedPtr<ISTSourceControlService> SourceControlService = ServiceProvider.Pin()->GetService<ISTSourceControlService>();
+		TSharedPtr<FJiraService> JiraService = ServiceProvider.Pin()->GetService<FJiraService>();
+		TSharedPtr<FTagService> TagService = ServiceProvider.Pin()->GetService<FTagService>();
+
+		UE_LOG(LogSubmitTool, Log, TEXT("Requesting Integration..."));
+		const FTag* JiraTag = TagService->GetTag(TEXT("#jira"));
+		FString JiraTagValue;
+
+		if(JiraTag != nullptr && JiraTag->GetValues().Num() > 0)
+		{
+			JiraTagValue = JiraTag->GetValues()[0];
+		}
+
+		if(!ChangelistService->GetFilesInCL().IsEmpty())
+		{
+			TFunction<void()> ShelveAndCreateSwarm = [this, ChangelistService, SourceControlService, SwarmService, TagService, JiraService, JiraTagValue, OnComplete]() {
+				ChangelistService->CreateShelvedFiles(FSourceControlOperationComplete::CreateLambda([this, ChangelistService, SourceControlService, SwarmService, TagService, JiraService, JiraTagValue, OnComplete](const FSourceControlOperationRef& DeleteShelvedOp, ECommandResult::Type Result)
+					{
+						if(Result == ECommandResult::Succeeded)
+						{
+							FString SwarmUrl;
+							if(!SwarmService->GetCurrentReviewUrl(SwarmUrl))
+							{
+								SwarmService->CreateReview(OnCreateReviewComplete::CreateLambda([this, ChangelistService, SourceControlService, JiraService, JiraTagValue, OnComplete](bool bSuccess, const FString& Response) {
+									if(bSuccess)
+									{
+										JiraService->GetIssueAndCreateServiceDeskRequest(JiraTagValue, ChangelistService->GetCLDescription(), Response, SourceControlService->GetCurrentStreamName(), IntegrationOptions, OnComplete);
+									}
+									else
+									{
+										UE_LOG(LogSubmitTool, Error, TEXT("Failed to create swarm review, Integration request is cancelled"));
+									}
+									}));
+							}
+							else
+							{
+								SwarmService->UpdateReviewDescription(FOnBooleanValueChanged::CreateLambda([this, ChangelistService, SourceControlService, JiraService, JiraTagValue, SwarmUrl, OnComplete](bool bSuccess)
+									{
+										if(!bSuccess)
+										{
+											UE_LOG(LogSubmitTool, Warning, TEXT("Failed to update swarm review description, Integration will continue with the current swarm description"));
+										}
+
+										JiraService->GetIssueAndCreateServiceDeskRequest(JiraTagValue, ChangelistService->GetCLDescription(), SwarmUrl, SourceControlService->GetCurrentStreamName(), IntegrationOptions, OnComplete);
+									}), ChangelistService->GetCLDescription());
+							}
+						}
+						else
+						{
+							UE_LOG(LogSubmitTool, Error, TEXT("Failed to shelve files, Integration request is cancelled"));
+						}
+					}));
+			};
+
+			if(ChangelistService->HasShelvedFiles())
+			{
+				EDialogFactoryResult Result = FDialogFactory::ShowConfirmDialog(FText::FromString(TEXT("Shelve files")), FText::FromString(TEXT("Submit tool will shelve your local files to update or create the swarm review.\nThis will replace the existing shelve entirely so you could lose any shelved changes that are not local, do you want to continue?")));
+				if(Result == EDialogFactoryResult::Confirm)
+				{
+					ChangelistService->DeleteShelvedFiles(FSourceControlOperationComplete::CreateLambda([this, ShelveAndCreateSwarm](const FSourceControlOperationRef& DeleteShelvedOp, ECommandResult::Type Result)
+						{
+							if(Result == ECommandResult::Succeeded)
+							{
+								ShelveAndCreateSwarm();
+							}
+							else
+							{
+								UE_LOG(LogSubmitTool, Error, TEXT("Failed to delete shelve, Integration request is cancelled"));
+							}
+						}));
+				}
+			}
+			else
+			{
+				ShelveAndCreateSwarm();
+			}
+		}
+		else if(!ChangelistService->GetShelvedFilesInCL().IsEmpty())
+		{
+			FString SwarmUrl;
+			if(SwarmService->GetCurrentReviewUrl(SwarmUrl))
+			{
+				SwarmService->UpdateReviewDescription(FOnBooleanValueChanged::CreateLambda([this, ChangelistService, SourceControlService, JiraService, JiraTagValue, SwarmUrl, OnComplete](bool bSuccess)
+					{
+						if(!bSuccess)
+						{
+							UE_LOG(LogSubmitTool, Error, TEXT("Failed to update swarm review description, Integration will continue with the current swarm description"));
+						}
+
+						JiraService->GetIssueAndCreateServiceDeskRequest(JiraTagValue, ChangelistService->GetCLDescription(), SwarmUrl, SourceControlService->GetCurrentStreamName(), IntegrationOptions, OnComplete);
+					}), ChangelistService->GetCLDescription());
+
+			}
+			else
+			{
+				SwarmService->CreateReview(OnCreateReviewComplete::CreateLambda([this, ChangelistService, SourceControlService, JiraService, JiraTagValue, OnComplete](bool bSuccess, const FString& Response) {
+					if(bSuccess)
+					{
+						JiraService->GetIssueAndCreateServiceDeskRequest(JiraTagValue, ChangelistService->GetCLDescription(), Response, SourceControlService->GetCurrentStreamName(), IntegrationOptions, OnComplete);
+					}
+					else
+					{
+						UE_LOG(LogSubmitTool, Error, TEXT("Failed to create swarm review, Integration request is cancelled"));
+					}
+					}));
+			}
+		}
+	}));
+}
+
+bool FIntegrationService::ValidateIntegrationOptions(bool bSilent) const
+{
+	bool bGlobalValidation = true;
+
+	const FTag* JiraTag = ServiceProvider.Pin()->GetService<FTagService>()->GetTag(TEXT("#jira"));
+	if(JiraTag != nullptr && JiraTag->GetValues().Num() > 0)
+	{
+		bool bPassesRegex = true;
+		FRegexPattern Pattern = FRegexPattern(JiraTag->Definition.Validation.RegexValidation, ERegexPatternFlags::CaseInsensitive);
+		for(const FString& Value : JiraTag->GetValues())
+		{
+			FRegexMatcher regex = FRegexMatcher(Pattern, Value);
+			bPassesRegex = regex.FindNext();
+			if(!bPassesRegex && !bSilent)
+			{
+				UE_LOG(LogSubmitTool, Error, TEXT("Value %s of Jira tag doesn't match the regex pattern %s"), *Value, *JiraTag->Definition.Validation.RegexValidation);
+			}
+		}
+
+		if(!bPassesRegex || JiraTag->GetValues()[0].Equals("none") || JiraTag->GetValues()[0].Equals("nojira"))
+		{
+			UE_LOG(LogSubmitTool, Error, TEXT("Validating Integration: Jira \"none\" and \"nojira\" are not allowed for Integrations, please specify a valid JIRA"));
+			bGlobalValidation = false;
+		}
+	}
+	else
+	{
+		UE_LOG(LogSubmitTool, Error, TEXT("Validating Integration: Jira is required for Integrations, please specify a valid JIRA"));
+		bGlobalValidation = false;
+	}
+
+	for(const TPair<FString, TSharedPtr<FIntegrationOptionBase>>& IntegrationOption : IntegrationOptions)
+	{
+		IntegrationOption.Value->bInvalid = false;
+	}
+
+	for(const FString& OrValidationGroup : Parameters.OneOfValidationGroups)
+	{
+		bool bValidGroup = false;
+		TArray<FString> GroupIds;
+
+		for(const TPair<FString, TSharedPtr<FIntegrationOptionBase>>& IntegrationOption : IntegrationOptions)
+		{
+			for(const FString& FieldValidationGroup : IntegrationOption.Value->FieldDefinition.ValidationGroups)
+			{
+				if(FieldValidationGroup.Equals(OrValidationGroup))
+				{
+					GroupIds.Add(IntegrationOption.Key);
+					FString out;
+					bValidGroup = bValidGroup || IntegrationOption.Value->GetJiraValue(out);
+					break;
+				}
+			}
+
+			if(bValidGroup)
+			{
+				break;
+			}
+		}
+
+		if(!bValidGroup)
+		{
+			if(!bSilent)
+			{
+				UE_LOG(LogSubmitTool, Error, TEXT("Validating Integration: One of these options need to have a value: %s"), *FString::Join(GroupIds, TEXT(", ")));
+			}
+
+			for(const FString& ID : GroupIds)
+			{
+				IntegrationOptions[ID]->bInvalid = true;
+			}
+
+			bGlobalValidation = false;
+		}
+	}
+
+
+	for(const TPair<FString, TSharedPtr<FIntegrationOptionBase>>& IntegrationOption : IntegrationOptions)
+	{
+		if(IntegrationOption.Value->FieldDefinition.bRequiredValue)
+		{
+			bool bValidDependencies = true;
+
+			if(IntegrationOption.Value->FieldDefinition.DependsOn.Num() > 0)
+			{
+				bValidDependencies = false;
+				for(const FString& Dependency : IntegrationOption.Value->FieldDefinition.DependsOn)
+				{
+					if(IntegrationOptions.Contains(Dependency))
+					{
+						FString ActualValue;
+						if(IntegrationOptions[Dependency]->GetJiraValue(ActualValue) || (!IntegrationOption.Value->FieldDefinition.DependsOnValue.IsEmpty() && IntegrationOption.Value->FieldDefinition.DependsOnValue.Equals(ActualValue)))
+						{
+							bValidDependencies = true;
+							break;
+						}
+					}
+				}
+			}
+
+			// Only Validate requirements if the dependencies are valid
+			if(bValidDependencies)
+			{
+				FString out;
+				if(!IntegrationOption.Value->GetJiraValue(out))
+				{
+					IntegrationOption.Value->bInvalid = true;
+					bGlobalValidation = false;
+
+					if(!bSilent)
+					{
+						UE_LOG(LogSubmitTool, Error, TEXT("Validating Integration: Option %s needs a value."), *IntegrationOption.Key);
+					}
+				}
+			}
+		}
+	}
+
+	return bGlobalValidation;
+}
diff --git a/Engine/Source/Programs/SubmitTool/Private/Logic/JiraService.cpp b/Engine/Source/Programs/SubmitTool/Private/Logic/JiraService.cpp
new file mode 100644
--- /dev/null
+++ b/Engine/Source/Programs/SubmitTool/Private/Logic/JiraService.cpp
@@ -0,0 +1,589 @@
+// Copyright Epic Games, Inc. All Rights Reserved.
+
+#include "JiraService.h"
+#include "GenericPlatform/GenericPlatformProcess.h"
+#include "Logic/DialogFactory.h"
+#include "Logging/SubmitToolLog.h"
+#include "SubmitToolUtils.h"
+#include "Logic/TagService.h"
+#include "Logic/PreflightService.h"
+#include "Logic/Services/SubmitToolServiceProvider.h"
+#include "Logic/Services/Interfaces/ISTSourceControlService.h"
+#include "Logic/CredentialsService.h"
+
+#include "HttpModule.h"
+#include "Interfaces/IHttpResponse.h"
+#include "Json.h"
+#include "Modules/ModuleManager.h"
+#include "Parameters/SubmitToolParameters.h"
+#include "JsonObjectConverter.h"
+#include "Configuration/Configuration.h"
+
+
+FJiraService::FJiraService(const FJiraParameters& InJiraSettings, const int32 InMaxResults, TSharedRef<FSubmitToolServiceProvider> InServiceProvider) :
+	Definition(InJiraSettings),
+	MaxResults(InMaxResults),
+	ServiceProvider(InServiceProvider)
+{
+	if(!Definition.ServerAddress.IsEmpty())
+	{
+		FetchJiraTickets(false);
+		LoadJiraIssues();
+	}
+}
+
+FJiraService::~FJiraService()
+{
+	if(JiraRequest.IsValid())
+	{
+		JiraRequest->CancelRequest();
+	}
+	OnJiraIssuesRetrievedCallback.Unbind();
+}
+
+bool FJiraService::FetchJiraTickets(bool InForce)
+{
+	TSharedPtr<FCredentialsService> CredentialsService = ServiceProvider->GetService<FCredentialsService>();
+	if(CredentialsService->HasCredentials() && !Definition.ServerAddress.IsEmpty())
+	{
+		if (InForce || (JiraIssues.Num() == 0 && CredentialsService->AreCredentialsValid()))
+		{
+			QueryIssues();
+			return true;
+		}
+	}
+
+	return false;
+}
+
+void FJiraService::Reset()
+{
+	JiraIssues.Reset();
+}
+
+void FJiraService::QueryIssues()
+{
+	if(JiraRequest.IsValid())
+	{
+		return;
+	}
+
+	FHttpModule& HttpModule = FModuleManager::LoadModuleChecked<FHttpModule>("HTTP");
+	JiraRequest = HttpModule.Get().CreateRequest();
+	
+	JiraRequest->OnProcessRequestComplete().BindRaw(this, &FJiraService::QueryIssues_HttpRequestComplete);
+
+	TSharedPtr<FCredentialsService> CredentialsService = ServiceProvider->GetService<FCredentialsService>();
+	FString Url = FString::Format(TEXT("https://{0}/rest/api/2/search?maxResults={1}&jql=assignee={2}"), { Definition.ServerAddress, this->MaxResults, CredentialsService->GetUsername()});
+
+	JiraRequest->SetURL(Url);
+	JiraRequest->SetHeader(TEXT("Authorization"), TEXT("Basic ") + CredentialsService->GetEncodedLoginString());
+	JiraRequest->SetVerb(TEXT("GET"));
+	bOngoingRequest = true;
+	UE_LOG(LogSubmitToolDebug, Log, TEXT("Sending Jira request for tickets assigned to %s"), *CredentialsService->GetUsername())
+	JiraRequest->ProcessRequest();
+}
+
+void FJiraService::QueryIssues_HttpRequestComplete(FHttpRequestPtr HttpRequest, FHttpResponsePtr HttpResponse, bool bSucceeded)
+{
+	JiraRequest = nullptr;
+	bOngoingRequest = false;
+	TSharedPtr<FCredentialsService> CredentialsService = ServiceProvider->GetService<FCredentialsService>();
+
+	if (!bSucceeded)
+	{
+		UE_LOG(LogSubmitToolDebug, Error, TEXT("Unable to retrieve JIRA issues at the moment."))
+		return;
+	}
+
+	if (HttpResponse.IsValid())
+	{
+		if (EHttpResponseCodes::IsOk(HttpResponse->GetResponseCode()))
+		{
+			CredentialsService->SetCredentialsValid(true);
+			UE_LOG(LogSubmitToolDebug, Log, TEXT("Successfully connected to Jira"));
+
+			FString ResponseStr = HttpResponse->GetContentAsString();
+
+			TSharedPtr<FJsonObject> RootJsonObject;
+			TSharedRef<TJsonReader<>> Reader = TJsonReaderFactory<>::Create(ResponseStr);
+			FJsonSerializer::Deserialize(Reader, RootJsonObject);
+
+			if (RootJsonObject.IsValid())
+			{
+				int32 Total;
+				if (RootJsonObject->TryGetNumberField(TEXT("total"), Total))
+				{
+					this->TotalIssues = Total;
+				}
+
+				const TArray<TSharedPtr<FJsonValue>>* Issues;
+				if (RootJsonObject->TryGetArrayField(TEXT("issues"), Issues))
+				{
+					UE_LOG(LogSubmitToolDebug, Log, TEXT("Retrieved %d issues for username %s"), Issues->Num(), *CredentialsService->GetUsername());
+
+					JiraIssues.Reset();
+
+					for (const TSharedPtr<FJsonValue>& ArrVal : *Issues)
+					{
+						if (!ArrVal.IsValid())
+						{
+							continue;
+						}
+
+						const TSharedPtr<FJsonObject>* IssueObject;
+						if (ArrVal->TryGetObject(IssueObject))
+						{
+							FJiraIssue Issue;
+							if(ParseJsonObject(IssueObject, Issue))
+							{
+								JiraIssues.Add(Issue.Key, Issue);
+							}
+						}
+					}
+
+					this->SaveJiraIssues();
+				}
+			}
+		}
+		else
+		{
+			UE_LOG(LogSubmitTool, Error, TEXT("Jira Request failed with error code %d, please make sure you're logging with the right credentials. if your Okta password expired recently, make sure you log into JIRA via browser at least once."), HttpResponse->GetResponseCode());
+			CredentialsService->SetCredentialsValid(false);
+		}
+	}
+
+	OnJiraIssuesRetrievedCallback.ExecuteIfBound(EHttpResponseCodes::IsOk(HttpResponse->GetResponseCode()));
+}
+
+bool FJiraService::ParseJsonObject(const TSharedPtr<FJsonObject>* InJsonObject, FJiraIssue& OutJiraIssue) const
+{
+	if(InJsonObject->IsValid())
+	{
+		FString Key{ "" };
+		FString Summary{ "" };
+		FString Description{ "" };
+		FString PriorityName{ "" };
+		FString StatusName{ "" };
+		FString IssueTypeName{ "" };
+
+		InJsonObject->Get()->TryGetStringField(TEXT("key"), Key);
+
+		const TSharedPtr<FJsonObject>* FieldsObject;
+		if(InJsonObject->Get()->TryGetObjectField(TEXT("fields"), FieldsObject))
+		{
+			if(FieldsObject->IsValid())
+			{
+				FieldsObject->Get()->TryGetStringField(TEXT("description"), Description);
+				FieldsObject->Get()->TryGetStringField(TEXT("summary"), Summary);
+
+				const TSharedPtr<FJsonObject>* PriorityObject;
+				if(FieldsObject->Get()->TryGetObjectField(TEXT("priority"), PriorityObject))
+				{
+					if(PriorityObject->IsValid())
+					{
+						PriorityObject->Get()->TryGetStringField(TEXT("name"), PriorityName);
+					}
+				}
+
+				const TSharedPtr<FJsonObject>* StatusObject;
+				if(FieldsObject->Get()->TryGetObjectField(TEXT("status"), StatusObject))
+				{
+					if(StatusObject->IsValid())
+					{
+						StatusObject->Get()->TryGetStringField(TEXT("name"), StatusName);
+					}
+				}
+
+				const TSharedPtr<FJsonObject>* IssueTypeObject;
+				if(FieldsObject->Get()->TryGetObjectField(TEXT("issuetype"), IssueTypeObject))
+				{
+					if(IssueTypeObject->IsValid())
+					{
+						IssueTypeObject->Get()->TryGetStringField(TEXT("name"), IssueTypeName);
+					}
+				}
+			}
+		}
+
+		if(!Key.IsEmpty() && !this->JiraIssues.Contains(Key))
+		{
+			FString Link = Definition.ServerAddress + TEXT("/browse/") + Key;
+			OutJiraIssue = FJiraIssue(Key, Summary, Link, Description, PriorityName, StatusName, IssueTypeName);
+			return true;
+		}
+	}
+	return false;
+}
+
+constexpr int JiraIssuesDatVersion = 1;
+void FJiraService::SaveJiraIssues() const
+{
+	FArchive* File = IFileManager::Get().CreateFileWriter(*GetJiraIssuesFilepath(), EFileWrite::FILEWRITE_EvenIfReadOnly);
+
+	if (File != nullptr)
+	{
+		int32 Version = JiraIssuesDatVersion;
+		*File << Version;
+
+		int32 Size = this->JiraIssues.Num();
+		*File << Size;
+
+		for (TPair<FString, FJiraIssue> Issue : this->JiraIssues)
+		{
+			FJiraIssue::StaticStruct()->SerializeBin(*File, &Issue.Value);
+		}
+
+		File->Close();
+		delete File;
+		File = nullptr;
+	}
+	else
+	{
+		UE_LOG(LogSubmitTool, Warning, TEXT("Could not create file '%s'."), *GetJiraIssuesFilepath());
+	}
+}
+
+void FJiraService::LoadJiraIssues()
+{
+	// do not load the issues if there is no credentials
+	if (!ServiceProvider->GetService<FCredentialsService>()->HasCredentials())
+	{
+		return;
+	}
+
+	if (IFileManager::Get().FileExists(*GetJiraIssuesFilepath()))
+	{
+		FArchive* File = IFileManager::Get().CreateFileReader(*GetJiraIssuesFilepath());
+
+		if (File != nullptr)
+		{
+			this->JiraIssues.Reset();
+
+			int32 Version;
+			*File << Version;
+
+			// Check Versions here
+			if (Version != JiraIssuesDatVersion)
+			{
+				UE_LOG(LogSubmitToolDebug, Warning, TEXT("Unexpected Jira Issues Version, aborting issues loading."));
+				File->Close();
+				delete File;
+				File = nullptr;
+				return;
+			}
+
+			int32 Size = 0;
+			*File << Size;
+
+			for (int32 Idx = 0; Idx < Size; Idx++)
+			{
+				FJiraIssue Issue;
+				FJiraIssue::StaticStruct()->SerializeBin(*File, &Issue);
+
+				if (!this->JiraIssues.Contains(Issue.Key))
+				{
+					this->JiraIssues.Add(Issue.Key, Issue);
+				}
+			}
+
+			File->Close();
+			delete File;
+			File = nullptr;
+		}
+		else
+		{
+			UE_LOG(LogSubmitTool, Warning, TEXT("Could not read file '%s'."), *GetJiraIssuesFilepath());
+		}
+	}
+	else
+	{
+		UE_LOG(LogSubmitToolDebug, Log, TEXT("File %s does not exists, no issues loaded"), *GetJiraIssuesFilepath())
+	}
+}
+
+const FString FJiraService::GetJiraIssuesFilepath() const
+{
+	return FPaths::Combine(FSubmitToolUtils::GetLocalAppDataPath(), TEXT("SubmitTool"), TEXT("jira.issues.dat"));
+}
+
+void FJiraService::GetuserInfo()
+{
+	FHttpModule& HttpModule = FModuleManager::LoadModuleChecked<FHttpModule>("HTTP");
+	FHttpRequestRef HttpRequest = HttpModule.Get().CreateRequest();
+
+	HttpRequest->OnProcessRequestComplete().BindRaw(this, &FJiraService::GetuserInfo_HttpRequestComplete);
+
+	FString Url = FString::Format(TEXT("https://{0}/rest/api/2/myself}"), { Definition.ServerAddress });
+
+	HttpRequest->SetURL(Url);
+
+	HttpRequest->SetHeader(TEXT("Authorization"), TEXT("Basic ") + ServiceProvider->GetService<FCredentialsService>()->GetEncodedLoginString());
+	HttpRequest->SetVerb(TEXT("GET"));
+	HttpRequest->ProcessRequest();
+}
+
+void FJiraService::GetuserInfo_HttpRequestComplete(FHttpRequestPtr HttpRequest, FHttpResponsePtr HttpResponse, bool bSucceeded)
+{
+	if (!bSucceeded)
+	{
+		UE_LOG(LogSubmitToolDebug, Error, TEXT("Unable to retrieve JIRA issues at the moment."))
+		return;
+	}
+
+	if (HttpResponse.IsValid())
+	{
+		if (EHttpResponseCodes::IsOk(HttpResponse->GetResponseCode()))
+		{
+			ServiceProvider->GetService<FCredentialsService>()->SetCredentialsValid(true);
+			FString ResponseStr = HttpResponse->GetContentAsString();
+
+			TSharedPtr<FJsonObject> RootJsonObject;
+			TSharedRef<TJsonReader<>> Reader = TJsonReaderFactory<>::Create(ResponseStr);
+			FJsonSerializer::Deserialize(Reader, RootJsonObject);
+
+			if (RootJsonObject.IsValid())
+			{
+			}
+		}
+	}
+}
+
+void FJiraService::GetIssueAndCreateServiceDeskRequest(const FString& Key, const FString& Description, const FString& SwarmURL, const FString& InCurrentStream, const TMap<FString, TSharedPtr<FIntegrationOptionBase>>& InIntegrationOptions, const FOnBooleanValueChanged OnComplete)
+{
+	if(Key.IsEmpty() || Key.Equals(TEXT("none"), ESearchCase::IgnoreCase))
+	{
+		CreateServiceDeskRequest(TSharedPtr<FJsonObject>(), Description, SwarmURL, InCurrentStream, InIntegrationOptions, OnComplete);
+		return;
+	}
+
+	UE_LOG(LogSubmitTool, Log, TEXT("Requesting Information for linked Jira issue %s"), *Key);
+	// I set this up so that it gets information from the linked jira, but I think they changed their minds on what information is taken from the jira
+	FHttpModule& HttpModule = FModuleManager::LoadModuleChecked<FHttpModule>("HTTP");
+	FHttpRequestRef HttpRequest = HttpModule.Get().CreateRequest();
+
+	HttpRequest->OnProcessRequestComplete().BindLambda([=, this](FHttpRequestPtr HttpRequest, FHttpResponsePtr HttpResponse, bool bSucceeded) {FJiraService::GetIssueAndCreateServiceDeskRequest_HttpRequestComplete(HttpRequest, HttpResponse, bSucceeded, Description, SwarmURL, InCurrentStream, InIntegrationOptions, OnComplete); });
+	
+	FString Url = FString::Format(TEXT("https://{0}/rest/api/2/issue/{1}"), { Definition.ServerAddress, Key });
+	HttpRequest->SetURL(Url);
+	HttpRequest->SetHeader(TEXT("Authorization"), TEXT("Basic ") + ServiceProvider->GetService<FCredentialsService>()->GetEncodedLoginString());
+	HttpRequest->SetVerb(TEXT("GET"));
+
+	HttpRequest->ProcessRequest();
+}
+
+void FJiraService::GetIssueAndCreateServiceDeskRequest_HttpRequestComplete(FHttpRequestPtr HttpRequest, FHttpResponsePtr HttpResponse, bool bSucceeded, const FString& Description, const FString& SwarmURL, const FString& InCurrentStream, const TMap<FString, TSharedPtr<FIntegrationOptionBase>>& InIntegrationOptions, const FOnBooleanValueChanged OnComplete)
+{
+	TSharedPtr<FJsonObject> RootJsonObject;
+	if (!bSucceeded)
+	{
+		if(HttpResponse.IsValid())
+		{
+			UE_LOG(LogSubmitTool, Log, TEXT("Unable to retrieve Base JIRA issue information. Summary will be created with the current CL description instead. Failed with code %d"), HttpResponse->GetResponseCode());
+			UE_LOG(LogSubmitToolDebug, Log, TEXT("Unable to retrieve JIRA issue information. Summary will be created with the current CL description instead. Failed with code %d\nResponse: %s"), HttpResponse->GetResponseCode(), *HttpResponse->GetContentAsString());
+		}
+		else
+		{
+			UE_LOG(LogSubmitTool, Warning, TEXT("Unable to retrieve Base JIRA issue information. Unknown failure"));
+		}
+	}
+	else
+	{
+		if(HttpResponse.IsValid())
+		{
+			if(EHttpResponseCodes::IsOk(HttpResponse->GetResponseCode()))
+			{
+				FString ResponseStr = HttpResponse->GetContentAsString();
+				UE_LOG(LogSubmitToolDebug, Log, TEXT("Obtained information from Jira Issue %s"), *ResponseStr);
+
+				TSharedRef<TJsonReader<>> Reader = TJsonReaderFactory<>::Create(ResponseStr);
+				FJsonSerializer::Deserialize(Reader, RootJsonObject);
+			}
+			else
+			{
+				UE_LOG(LogSubmitTool, Warning, TEXT("Unable to retrieve Base JIRA issue information."));
+				UE_LOG(LogSubmitToolDebug, Warning, TEXT("Unable to retrieve Base JIRA issue information. Failed with code %d\nResponse: %s"), HttpResponse->GetResponseCode(), *HttpResponse->GetContentAsString());
+			}
+		}
+	}
+
+	// Call the function to actually create the service desk request with the required information
+	CreateServiceDeskRequest(RootJsonObject, Description, SwarmURL, InCurrentStream, InIntegrationOptions, OnComplete);
+}
+
+void FJiraService::CreateServiceDeskRequest(TSharedPtr<FJsonObject> InBaseJiraJsonObject, const FString& Description, const FString& SwarmURL, const FString& InCurrentStream, const TMap<FString, TSharedPtr<FIntegrationOptionBase>>& InIntegrationOptions, const FOnBooleanValueChanged OnComplete)
+{
+	UE_LOG(LogSubmitTool, Log, TEXT("Requesting creation of Jira ServiceDesk ticket..."));
+	FHttpModule& HttpModule = FModuleManager::LoadModuleChecked<FHttpModule>("HTTP");
+	ServiceDeskRequest = HttpModule.Get().CreateRequest();
+
+	ServiceDeskRequest->OnProcessRequestComplete().BindLambda([this, OnComplete](FHttpRequestPtr HttpRequest, FHttpResponsePtr HttpResponse, bool bSucceeded) { CreateServiceDeskRequest_HttpRequestComplete(HttpRequest, HttpResponse, bSucceeded, OnComplete); });
+	FString Url = FString::Format(TEXT("https://{0}/rest/servicedeskapi/request"), { Definition.ServerAddress });
+
+	ServiceDeskRequest->SetURL(Url);
+	ServiceDeskRequest->SetHeader(TEXT("Authorization"), FString::Format(TEXT("Basic {0}"), { Definition.ServiceDeskToken }));
+	ServiceDeskRequest->SetHeader(TEXT("Content-Type"), TEXT("application/json"));
+	ServiceDeskRequest->SetVerb(TEXT("POST"));
+	
+	TSharedPtr<FJsonObject> RequestJson = MakeShared<FJsonObject>();
+
+	// These values should probably be put inside the .ini file
+	RequestJson->SetNumberField(TEXT("serviceDeskId"), Definition.ServiceDeskID);
+	RequestJson->SetNumberField(TEXT("requestTypeId"), Definition.RequestFormID);
+
+
+	TSharedPtr<FJsonObject> RequestFieldValuesJson = MakeShared<FJsonObject>();
+
+	if(InBaseJiraJsonObject.IsValid())
+	{
+		TSharedPtr<FJsonObject> BaseJiraFields = InBaseJiraJsonObject->GetObjectField(TEXT("fields"));
+		RequestFieldValuesJson->SetStringField(TEXT("summary"), BaseJiraFields->GetStringField(TEXT("summary")));
+	}
+	else
+	{
+		FString Summary = Description.Left(50).Replace(TEXT("\n"), TEXT(" ")).Replace(TEXT("\r"), TEXT(""));
+		RequestFieldValuesJson->SetStringField(TEXT("summary"), Summary);
+
+	}
+
+	RequestFieldValuesJson->SetStringField(TEXT("description"), Description);
+
+	if(!SwarmURL.IsEmpty() && !Definition.SwarmUrlField.IsEmpty())
+	{
+		RequestFieldValuesJson->SetStringField(Definition.SwarmUrlField, SwarmURL);
+	}
+
+	if(!InCurrentStream.IsEmpty() && !Definition.StreamField.IsEmpty())
+	{
+		RequestFieldValuesJson->SetStringField(Definition.StreamField, InCurrentStream);
+	}
+
+	if(!Definition.PreflightField.IsEmpty())
+	{		
+		const FTag* PreflightTag = ServiceProvider->GetService<FTagService>()->GetTag(TEXT("#preflight"));
+		if (PreflightTag != nullptr && PreflightTag->GetValues().Num() > 0)
+		{
+			FString PreflightTagValue = PreflightTag->GetValues()[0];
+
+			if(!PreflightTagValue.IsEmpty())
+			{
+				if(!PreflightTagValue.Contains(TEXT("/job/")))
+				{
+					RequestFieldValuesJson->SetStringField(Definition.PreflightField, FString::Format(TEXT("{0}job/{1}"), { ServiceProvider->GetService<FPreflightService>()->GetHordeServerAddress(), PreflightTagValue }));
+				}
+				else
+				{
+					RequestFieldValuesJson->SetStringField(Definition.PreflightField, PreflightTagValue);
+				}
+			}
+		}
+	}
+
+	FString Requestor = FConfiguration::Substitute(TEXT("$(USER)"));
+	TSharedPtr<FUserData> LocalUserData = ServiceProvider->GetService<ISTSourceControlService>()->GetUserDataFromCache(Requestor);
+	if (LocalUserData.IsValid())
+	{
+		Requestor = LocalUserData.Get()->Email;
+	}
+
+	RequestFieldValuesJson->SetStringField(Definition.RequestorField, Requestor);
+
+	for(const TPair<FString, TSharedPtr<FIntegrationOptionBase>>& Pair : InIntegrationOptions)
+	{
+		FString Value;
+		if(!Pair.Value->GetJiraValue(Value))
+		{
+			continue;
+		}
+
+		TSharedPtr<FJsonObject> JiraObject = MakeShared<FJsonObject>();
+		TArray<TSharedPtr<FJsonValue>> JiraArrayObject;
+		JiraObject->SetStringField(TEXT("value"), Value);
+
+		switch(Pair.Value->FieldDefinition.JiraType)
+		{
+		case EJiraFieldType::Object:
+			RequestFieldValuesJson->SetObjectField(Pair.Value->FieldDefinition.Id, JiraObject);
+			break;
+
+		case EJiraFieldType::Array:
+			const TArray<TSharedPtr<FJsonValue>>* ExistingJiraArrayObjectPtr;
+			if(RequestFieldValuesJson->TryGetArrayField(Pair.Value->FieldDefinition.Id, ExistingJiraArrayObjectPtr))
+			{
+				JiraArrayObject = *ExistingJiraArrayObjectPtr;
+			}
+
+			JiraArrayObject.Add(MakeShared<FJsonValueObject>(JiraObject));
+			RequestFieldValuesJson->SetArrayField(Pair.Value->FieldDefinition.Id, JiraArrayObject);
+			break;
+
+		case EJiraFieldType::String:
+			RequestFieldValuesJson->SetStringField(Pair.Value->FieldDefinition.Id, Value);
+			break;
+		}
+	}
+
+	// JW: I think they decided that any jira tickets referenced in the changelist should be added as additional URLs
+	RequestJson->SetObjectField(TEXT("requestFieldValues"), RequestFieldValuesJson);
+
+	FString BodyString;
+	TSharedRef<TJsonWriter<TCHAR>> JsonWriter = TJsonWriterFactory<>::Create(&BodyString);
+	FJsonSerializer::Serialize(RequestJson.ToSharedRef(), JsonWriter);
+
+	UE_LOG(LogSubmitToolDebug, Log, TEXT("Create Jira request body:\n%s"), *BodyString);
+
+	ServiceDeskRequest->SetContentAsString(BodyString);
+
+	ServiceDeskRequest->ProcessRequest();
+}
+
+void FJiraService::CreateServiceDeskRequest_HttpRequestComplete(FHttpRequestPtr HttpRequest, FHttpResponsePtr HttpResponse, bool bSucceeded, const FOnBooleanValueChanged OnComplete)
+{
+	if(!bSucceeded)
+	{
+		if(HttpResponse.IsValid())
+		{
+			UE_LOG(LogSubmitTool, Error, TEXT("Unable to create JIRA service desk. Failed with code %d\nResponse: %s"), HttpResponse->GetResponseCode(), *HttpResponse->GetContentAsString());
+		}
+		else
+		{
+			UE_LOG(LogSubmitTool, Error, TEXT("Unable to create JIRA service desk. Unknown failure"));
+		}
+		OnComplete.ExecuteIfBound(false);
+		return;
+	}
+
+	if (HttpResponse.IsValid())
+	{
+		if (EHttpResponseCodes::IsOk(HttpResponse->GetResponseCode()))
+		{
+			TSharedRef<TJsonReader<TCHAR>> Reader = TJsonReaderFactory<TCHAR>::Create(HttpResponse->GetContentAsString());
+			TSharedPtr<FJsonObject> JsonObj;
+			if(!FJsonSerializer::Deserialize(Reader, JsonObj))
+			{
+				UE_LOG(LogSubmitTool, Error, TEXT("Unable to deserialize swarm create response"));
+				return;
+			}
+
+			FString CreatedTicketId = JsonObj->GetStringField(TEXT("issueKey"));
+			FString WebLink;
+			const TSharedPtr<FJsonObject>* Links;
+			if(JsonObj->TryGetObjectField(TEXT("_links"), Links))
+			{
+				WebLink = Links->Get()->GetStringField(TEXT("web"));
+			}
+
+			// If the service desk request was created successfully
+			UE_LOG(LogSubmitTool, Log, TEXT("Jira service desk ticket creation was successful: %s %s"), *CreatedTicketId, *WebLink);
+			UE_LOG(LogSubmitToolDebug, Log, TEXT("Jira service desk ticket creation was successful\n%s"), *HttpResponse->GetContentAsString());
+			FDialogFactory::ShowInformationDialog(FText::FromString(TEXT("Integration Request Successful")), FText::FromString(TEXT("The Integration has sucessfully been requested!")));			
+			OnComplete.ExecuteIfBound(true);
+		}
+		else
+		{
+			UE_LOG(LogSubmitTool, Error, TEXT("Unable to create JIRA service desk. Failed with code %d\nResponse: %s"), HttpResponse->GetResponseCode(), *HttpResponse->GetContentAsString());
+			FDialogFactory::ShowInformationDialog(FText::FromString(TEXT("Integration Request FAILED")), FText::FromString(TEXT("Unable to create JIRA service desk ticket.\nPlease check the logs for more info.")));
+			OnComplete.ExecuteIfBound(false);
+		}
+	}
+}
diff --git a/Engine/Source/Programs/SubmitTool/Private/Logic/P4LockdownService.cpp b/Engine/Source/Programs/SubmitTool/Private/Logic/P4LockdownService.cpp
new file mode 100644
--- /dev/null
+++ b/Engine/Source/Programs/SubmitTool/Private/Logic/P4LockdownService.cpp
@@ -0,0 +1,336 @@
+// Copyright Epic Games, Inc. All Rights Reserved.
+
+#include "P4LockdownService.h"
+#include "CommandLine/CmdLineParameters.h"
+#include "Logic/Services/Interfaces/ISTSourceControlService.h"
+#include "Logic/Services/SubmitToolServiceProvider.h"
+#include "Logging/SubmitToolLog.h"
+#include "Misc/ConfigCacheIni.h"
+#include "Misc/FileHelper.h"
+#include "Misc/Paths.h"
+#include "HAL/FileManager.h"
+#include "Internationalization/Regex.h"
+#include "Configuration/Configuration.h"
+
+FP4LockdownService::FP4LockdownService(const FP4LockdownParameters& InParameters, TSharedPtr<FSubmitToolServiceProvider> InServiceProvider)
+: 	ServiceProvider(InServiceProvider), Parameters(InParameters)
+{
+	using namespace UE::Tasks;
+	ConfigFilesTask = Launch(UE_SOURCE_LOCATION, [this]
+		{
+			TArray<TTask<bool>> DownloadFileTasks;
+			for(const TPair<FString, FString>& Pair : Parameters.ConfigPaths)
+			{
+				DownloadFileTasks.Add(Launch(UE_SOURCE_LOCATION, [this, &Pair]
+					{
+						return GetConfigFile(Pair.Key, Pair.Value);
+					}));
+			}
+
+			Wait(DownloadFileTasks);
+
+			bool bResult = true;
+			for(TTask<bool>& Task : DownloadFileTasks)
+			{
+				bResult &= Task.GetResult();
+			}
+
+			return bResult;
+		});
+}
+
+bool FP4LockdownService::ArePathsInLockdown(const TArray<FString>& InPaths, bool& bOutAllowlisted)
+{	
+	if(!ConfigFilesTask.IsValid())
+	{
+		UE_LOG(LogSubmitToolP4, Error, TEXT("Downloading task wasn't setup correctly, hardlock status is not known."));
+		return false;
+	}
+
+	if(!ConfigFilesTask.IsCompleted())
+	{		
+		UE_LOG(LogSubmitToolP4, Log, TEXT("Waiting for download of Stream Hardlock data..."));
+		ConfigFilesTask.Wait(FTimespan::FromSeconds(5));
+
+		if(!ConfigFilesTask.IsCompleted())
+		{
+			UE_LOG(LogSubmitToolP4, Error, TEXT("Downloading config files from P4 timed out, hardlock status is not latest, will use cache."));
+		}
+	}
+
+	if(ConfigFilesTask.IsCompleted())
+	{
+		for(const TPair<FString, FSharedBuffer>& Pair : DownloadedFiles)
+		{
+			const FString Path = GetFilePath(Pair.Key);
+			FArchive* File = IFileManager::Get().CreateFileWriter(*Path, FILEWRITE_EvenIfReadOnly);
+
+			if(File != nullptr)
+			{
+				File->Serialize(const_cast<void*>(Pair.Value.GetData()), Pair.Value.GetSize());
+				File->Close();
+				delete File;
+			}
+			else
+			{
+				UE_LOG(LogSubmitTool, Warning, TEXT("Couldn't create lockdown file %s"), *Path)
+			}
+		}
+	}
+
+	if(!ConfigFilesTask.GetResult())
+	{
+		UE_LOG(LogSubmitToolP4, Warning, TEXT("Downloading config files from P4 had failures. Lockdown data won't be complete"));
+	}
+
+	ParseAllowListData();
+
+	if(InPaths.IsEmpty())
+	{
+		UE_LOG(LogSubmitToolP4, Warning, TEXT("No files to check for lockdown"));
+	}
+
+	bool bLockdown = false;
+	bOutAllowlisted = true;
+	for(const FString& Path : InPaths)
+	{
+		FPathLockdownResult Result = IsPathInLockdown(Path);
+		bLockdown |= Result.bIsLocked;
+
+		if(Result.bIsLocked)
+		{
+			bOutAllowlisted &= Result.bAllowlisted;
+		}
+	}
+
+	return bLockdown;
+}
+
+FP4LockdownService::FPathLockdownResult FP4LockdownService::IsPathInLockdown(const FString& InPath) const
+{
+	FString PerforceUserName;
+	FCmdLineParameters::Get().GetValue(FSubmitToolCmdLine::P4User, PerforceUserName);
+
+	FPathLockdownResult bOverallLockdownResult { false, true };
+
+	for(const FAllowListData& data : AllowListData)
+	{
+		bool bIsLocked = false;
+
+		auto EvaluateViews = [&bIsLocked, &InPath](const TArray<TPair<bool, FString>>& Views)
+		{
+			for (const TPair<bool, FString>& View : Views)
+			{
+				if (bIsLocked != View.Key)
+				{
+					FRegexPattern Pattern = FRegexPattern(View.Value, ERegexPatternFlags::CaseInsensitive);
+					FRegexMatcher regex = FRegexMatcher(MoveTemp(Pattern), InPath);
+
+					if (regex.FindNext())
+					{
+						bIsLocked = View.Key;
+					}
+				}
+			}
+		};
+		EvaluateViews(data.Views);
+
+		bool bIsInOverrideAllowlist = false;
+		for (const FOverrideData& Override : OverrideData)
+		{
+			if (Override.Sections.Contains(data.GroupName) && Override.AllowListers.Contains(PerforceUserName))
+			{
+				EvaluateViews(Override.Views);
+				bIsInOverrideAllowlist = true;
+			}
+		}
+
+		if (bIsLocked)
+		{
+			bOverallLockdownResult.bAllowlisted = bIsInOverrideAllowlist || data.AllowListers.Contains(PerforceUserName);
+			bOverallLockdownResult.bIsLocked = true;
+			break;
+		}
+	}
+
+	for(const FString& AdditionalHardlockedPath : AdditionalHardlocks)
+	{
+		FRegexPattern Pattern = FRegexPattern(AdditionalHardlockedPath, ERegexPatternFlags::CaseInsensitive);
+		FRegexMatcher regex = FRegexMatcher(MoveTemp(Pattern), InPath);
+
+		if(regex.FindNext())
+		{
+			bOverallLockdownResult.bIsLocked = true;
+			bOverallLockdownResult.bAllowlisted = false;
+			break;
+		}
+	}
+
+	return bOverallLockdownResult;
+}
+
+bool FP4LockdownService::GetConfigFile(const FString& InConfigId, const FString& InDepotPath)
+{
+	TArray<FSharedBuffer> FileBuffers;
+	if(ServiceProvider->GetService<ISTSourceControlService>()->DownloadFiles(InDepotPath, FileBuffers).GetResult() && FileBuffers.Num() > 0)
+	{
+		FScopeLock Scope(&FileMutex);
+		
+		DownloadedFiles.Add(InConfigId, FileBuffers[0]);
+		return true;
+	}
+
+	return false;
+}
+
+void RegexEscapeInline(FString& InOutRegex)
+{
+	for (int i = InOutRegex.Len() - 1; i >= 0; --i)
+	{
+		switch (InOutRegex[i])
+		{
+			case TEXT('\\'):
+			case TEXT('*'):
+			case TEXT('+'):
+			case TEXT('?'):
+			case TEXT('|'):
+			case TEXT('{'):
+			case TEXT('}'):
+			case TEXT('['):
+			case TEXT(']'):
+			case TEXT('('):
+			case TEXT(')'):
+			case TEXT('^'):
+			case TEXT('$'):
+			case TEXT('.'):
+			case TEXT('#'):
+			case TEXT(' '):
+				InOutRegex.InsertAt(i, TEXT('\\'));
+				break;
+		}
+	}
+}
+
+void FP4LockdownService::GetAdditionalHardlockedPaths()
+{
+	for(FString EscapedViewLine : Parameters.AdditionalHardlockedPaths)
+	{
+		RegexEscapeInline(EscapedViewLine);
+		EscapedViewLine.ReplaceInline(TEXT("\\*"), TEXT("[^/]*"));
+		EscapedViewLine.ReplaceInline(TEXT("\\.\\.\\."), TEXT(".*"));
+
+		AdditionalHardlocks.Add(MoveTemp(EscapedViewLine));
+	}
+}
+
+void FP4LockdownService::ParseAllowListData()
+{
+	for(const TPair<FString, FString>& Config : Parameters.ConfigPaths)
+	{
+		FString Filepath = GetFilePath(Config.Key);
+		if(IFileManager::Get().FileExists(*Filepath))
+		{
+			FConfigFile LockdownConfig;
+			LockdownConfig.bPythonConfigParserMode = true;
+			LockdownConfig.Read(Filepath);
+
+			for (const TPair<FString, FConfigSection>& ConfigPair : AsConst(LockdownConfig))
+			{
+				const FConfigSection& ConfigSection = ConfigPair.Value;
+				
+				const FConfigValue* AllowList = ConfigSection.Find(TEXT("allowlist"));
+				if (AllowList == nullptr)
+				{
+					continue;
+				}
+
+				const FConfigValue* Status = ConfigSection.Find(TEXT("status"));
+				if (Status == nullptr)
+				{
+					continue;
+				}
+
+				FAllowListData* data;
+				if (Status->GetSavedValue() == TEXT("hardcore"))
+				{
+					data = &AllowListData.AddDefaulted_GetRef();
+				}
+				else if (Status->GetSavedValue() == TEXT("override"))
+				{
+					const FConfigValue* SectionList = ConfigSection.Find(TEXT("sectionlist"));
+					if (SectionList == nullptr)
+					{
+						continue;
+					}
+
+					FOverrideData& Override = OverrideData.AddDefaulted_GetRef();
+					data = &Override;
+
+					TArray<FString> Sections;
+					SectionList->GetSavedValue().ParseIntoArray(Sections, TEXT(","), true);
+
+					for(FString& Section : Sections)
+					{
+						Section.ToLowerInline();
+						Override.Sections.Add(MoveTemp(Section));
+					}
+				}
+				else
+				{
+					continue;
+				}
+
+				data->GroupName = ConfigPair.Key;
+
+				TArray<FString> Allowlisters;
+				AllowList->GetSavedValue().ParseIntoArray(Allowlisters, TEXT(","), true);
+
+				for(FString& AllowLister : Allowlisters)
+				{
+					AllowLister.ToLowerInline();
+					data->AllowListers.Add(MoveTemp(AllowLister));
+				}
+
+				TArray<FString> View;
+				ConfigSection.MultiFind(TEXT("view"), View, true);
+				for(FString& ViewLine : View)
+				{
+					if(ViewLine.IsEmpty())
+					{
+						continue;
+					}
+
+					bool bIsLocked = true;
+					if(ViewLine.StartsWith(TEXT("-"), ESearchCase::CaseSensitive))
+					{
+						bIsLocked = false;
+						ViewLine.MidInline(2);
+					}
+
+					RegexEscapeInline(ViewLine);
+					ViewLine.ReplaceInline(TEXT("\\*"), TEXT("[^/]*"), ESearchCase::CaseSensitive);
+					ViewLine.ReplaceInline(TEXT("\\.\\.\\."), TEXT(".*"), ESearchCase::CaseSensitive);
+
+					data->Views.Add(TPair<bool, FString>(bIsLocked, MoveTemp(ViewLine)));
+				}
+			}
+		}
+		else
+		{
+			UE_LOG(LogSubmitToolP4Debug, Error, TEXT("File %s doesn't exist"), *Filepath);
+		}
+	}
+}
+
+FString FP4LockdownService::GetFilePath(const FString& InConfigId)
+{
+	FString EngineDir = FPaths::ConvertRelativePathToFull(FPaths::EngineDir());
+	FPaths::NormalizeDirectoryName(EngineDir);
+	
+	const FGuid guid = FGuid::NewGuid();
+	FString LocalFilePath = EngineDir + TEXT("/Intermediate/SubmitTool/P4Lockdown/") + InConfigId + TEXT(".ini");
+	LocalFilePath = FPaths::ConvertRelativePathToFull(LocalFilePath);
+	FPaths::NormalizeFilename(LocalFilePath);
+		
+	return LocalFilePath;
+}
\ No newline at end of file
diff --git a/Engine/Source/Programs/SubmitTool/Private/Logic/P4LockdownService.h b/Engine/Source/Programs/SubmitTool/Private/Logic/P4LockdownService.h
new file mode 100644
--- /dev/null
+++ b/Engine/Source/Programs/SubmitTool/Private/Logic/P4LockdownService.h
@@ -0,0 +1,63 @@
+// Copyright Epic Games, Inc. All Rights Reserved.
+
+#pragma once
+
+#include "CoreMinimal.h"
+#include "Parameters/SubmitToolParameters.h"
+#include "Logic/ProcessWrapper.h"
+#include "Tasks/Task.h"
+#include "Services/Interfaces/ISubmitToolService.h"
+#include "Memory/SharedBuffer.h"
+
+DECLARE_DELEGATE_OneParam(FOnConfigFileRetrieved, bool /*bValid*/)
+
+class FSubmitToolServiceProvider;
+
+class FP4LockdownService : public ISubmitToolService
+{	
+public:
+	FP4LockdownService() = delete;
+	FP4LockdownService(const FP4LockdownParameters& InParameters, TSharedPtr<FSubmitToolServiceProvider> InServiceProvider);
+	bool ArePathsInLockdown(const TArray<FString>& InPaths, bool& bOutAllowlisted);
+	bool IsBlockingOperationRunning() const
+	{
+		return ConfigFilesTask.IsValid() && !ConfigFilesTask.IsCompleted();
+	}
+
+private:
+
+	struct FAllowListData
+	{
+		FString GroupName;
+		TSet<FString> AllowListers;
+		TArray<TPair<bool, FString>> Views;
+	};
+
+	struct FOverrideData : public FAllowListData
+	{
+		TSet<FString> Sections;
+	};
+
+	struct FPathLockdownResult
+	{
+		bool bIsLocked;
+		bool bAllowlisted;
+	};
+
+	FPathLockdownResult IsPathInLockdown(const FString& InPath) const;
+	bool GetConfigFile(const FString& InConfigId, const FString& InDepotPath);
+	void ParseAllowListData();
+	void GetAdditionalHardlockedPaths();
+	FString GetFilePath(const FString& InConfigId);
+
+	TMap<FString, FSharedBuffer> DownloadedFiles;
+	UE::Tasks::TTask<bool> ConfigFilesTask;
+	TSharedPtr<FSubmitToolServiceProvider> ServiceProvider;
+	const FP4LockdownParameters& Parameters;
+	TArray<FAllowListData> AllowListData;
+	TArray<FOverrideData> OverrideData;
+	TArray<FString> AdditionalHardlocks;
+	FCriticalSection FileMutex;
+};
+
+Expose_TNameOf(FP4LockdownService);
\ No newline at end of file
