Moving Submit Tool to Source/Programs

diff --git a/Engine/Source/Programs/SubmitTool/Private/Logic/PreflightService.cpp b/Engine/Source/Programs/SubmitTool/Private/Logic/PreflightService.cpp
new file mode 100644
--- /dev/null
+++ b/Engine/Source/Programs/SubmitTool/Private/Logic/PreflightService.cpp
@@ -0,0 +1,700 @@
+// Copyright Epic Games, Inc. All Rights Reserved.
+
+#include "PreflightService.h"
+#include "GenericPlatform/GenericPlatformProcess.h"
+#include "HttpModule.h"
+#include "Interfaces/IHttpResponse.h"
+#include "JsonObjectConverter.h"
+#include "Misc/Paths.h"
+#include "Modules/ModuleManager.h"
+#include "Serialization/JsonSerializer.h"
+#include "Widgets/Docking/SDockTab.h"
+#include "Internationalization/Regex.h"
+
+#include "CommandLine/CmdLineParameters.h"
+#include "Configuration/Configuration.h"
+#include "Logic/Services/Interfaces/ISTSourceControlService.h"
+#include "Logic/CredentialsService.h"
+#include "Logic/TagService.h"
+#include "Logging/SubmitToolLog.h"
+#include "Models/ModelInterface.h"
+#include "Parameters/SubmitToolParameters.h"
+
+FPreflightService::FPreflightService(
+	const FHordeParameters& InSettings,
+	FModelInterface* InModelInterface,
+	TSharedRef<FSubmitToolServiceProvider> InServiceProvider)
+	:
+	Definition(InSettings),
+	ServiceProvider(InServiceProvider),
+	ModelInterface(InModelInterface),
+	State(EPreflightServiceState::Idle),
+	LastErrorMessage(TEXT(""))
+{
+	TickHandle = FTSTicker::GetCoreTicker().AddTicker(FTickerDelegate::CreateRaw(this, &FPreflightService::Tick));
+
+	DeleteShelveCallback = FSourceControlOperationComplete::CreateRaw(this, &FPreflightService::OnDeleteShelveOperationComplete);
+	CreateShelveCallback = FSourceControlOperationComplete::CreateRaw(this, &FPreflightService::OnCreateShelveOperationComplete);
+
+	PreflightTag = ServiceProvider->GetService<FTagService>()->GetTag(TEXT("#preflight"));
+}
+
+FPreflightService::~FPreflightService()
+{
+	FTSTicker::GetCoreTicker().RemoveTicker(TickHandle);
+	OnPreflightDataUpdated.Clear();
+}
+
+bool FPreflightService::Tick(float DeltaTime)
+{
+	TSharedPtr<FChangelistService> ChangelistService = ServiceProvider->GetService<FChangelistService>();
+	switch(State)
+	{
+		//////////////////////////////////////////////////
+		case EPreflightServiceState::Idle:
+			// Do nothing, wait for someone to press the "Start" preflight button
+			break;
+
+			//////////////////////////////////////////////////
+		case EPreflightServiceState::RequestDeleteShelve:
+
+			if(ChangelistService->GetFilesInCL().Num() == 0)
+			{
+				LastErrorMessage = TEXT("Missing local files!  We can't update the shelved files!");
+				State = EPreflightServiceState::Error;
+			}
+			else if(ChangelistService->HasShelvedFiles())
+			{
+				EDialogFactoryResult DialogResult = static_cast<EDialogFactoryResult>(ShowRecreateShelveDialog());
+				if(DialogResult == EDialogFactoryResult::FirstButton)
+				{
+					UE_LOG(LogSubmitTool, Log, TEXT("Preflight: Deleting shelved files"));
+					State = EPreflightServiceState::WaitingForDeleteShelve;
+					bCheckShelveInstead = false;
+					ChangelistService->DeleteShelvedFiles(DeleteShelveCallback);
+				}
+				else if(DialogResult == EDialogFactoryResult::SecondButton)
+				{
+					UE_LOG(LogSubmitTool, Log, TEXT("Preflight: Using existing shelved files"));
+					bCheckShelveInstead = true;
+					State = EPreflightServiceState::StartPreflight;
+				}
+				else
+				{
+					UE_LOG(LogSubmitTool, Log, TEXT("Preflight: Request cancelled"));
+					State = EPreflightServiceState::Idle;
+				}
+			}
+			else
+			{
+				State = EPreflightServiceState::RequestCreateShelve;
+			}
+			break;
+
+			//////////////////////////////////////////////////
+		case EPreflightServiceState::WaitingForDeleteShelve:
+			// Do nothing, wait for the delete shelve callback
+			break;
+
+			//////////////////////////////////////////////////
+		case EPreflightServiceState::RequestCreateShelve:
+			UE_LOG(LogSubmitTool, Log, TEXT("Preflight: Creating shelved files..."));
+			State = EPreflightServiceState::WaitingForCreateShelve;
+			ChangelistService->CreateShelvedFiles(CreateShelveCallback);
+			break;
+
+			//////////////////////////////////////////////////
+		case EPreflightServiceState::WaitingForCreateShelve:
+			// Do nothing, wait for the create shelve callback
+			break;
+
+			//////////////////////////////////////////////////
+		case EPreflightServiceState::StartPreflight:
+			StartPreflight();
+			break;
+
+			//////////////////////////////////////////////////
+		case EPreflightServiceState::Error:
+			UE_LOG(LogSubmitTool, Error, TEXT("Preflight: \"%s\""), *LastErrorMessage);
+			State = EPreflightServiceState::Idle;
+			break;
+	}
+
+	return true;
+}
+
+TMap<FString, FStringFormatArg> FPreflightService::GetFormatParameters() const
+{
+	TMap<FString, FStringFormatArg> FormatMap =
+	{
+		{ TEXT("URL"), Definition.HordeServerAddress },
+		{ TEXT("CLID"), ServiceProvider->GetService<FChangelistService>()->GetCLID()}
+	};
+
+	FormatMap.Add(TEXT("Stream"), StreamName);
+	FormatMap.Add(TEXT("Template"), FString());
+	FormatMap.Add(TEXT("AdditionalTasks"), FString());
+
+	FPreflightTemplateDefinition Template;
+	if(SelectPreflightTemplate(Template))
+	{
+		FormatMap[TEXT("Template")] = Template.Template;
+		FormatMap[TEXT("AdditionalTasks")] = GetAdditionalTasksString(Template);
+	}
+	else
+	{
+		FormatMap[TEXT("Template")] = Definition.DefaultPreflightTemplate;
+	}
+
+	return FormatMap;
+}
+
+void FPreflightService::RequestPreflight()
+{
+	if(State == EPreflightServiceState::Idle)
+	{
+		if(!ServiceProvider->GetService<ISTSourceControlService>()->GetClientStreams().IsEmpty())
+		{
+			UE_LOG(LogSubmitTool, Log, TEXT("Preflight: Requesting..."));
+			State = EPreflightServiceState::RequestDeleteShelve;
+		}
+		else
+		{
+			UE_LOG(LogSubmitTool, Error, TEXT("Couldn't retrieve stream name in this p4 client. Submit tool can't start a preflight, see previous errors."));
+		}
+	}
+	else
+	{
+		// Do nothing, we're already busy trying to start a preflight
+	}
+}
+
+void FPreflightService::QueueFetch(bool bRequeue, float InSeconds)
+{
+	FTSTicker::GetCoreTicker().AddTicker(FTickerDelegate::CreateLambda([this, bRequeue](float DeltaTime) { FetchPreflightInfo(bRequeue); return false; }), InSeconds);
+}
+
+void FPreflightService::Requeue()
+{
+	float WaitTime = Definition.FetchPreflightEachSeconds;
+
+	if(HordePreflights.IsValid())
+	{
+		for(const FPreflightData& PFData : HordePreflights->PreflightList)
+		{
+			if(PFData.CachedResults.State != EPreflightState::Completed)
+			{
+				WaitTime = Definition.FetchPreflightEachSecondsWhenInProgress;
+				break;
+			}
+		}
+	}
+
+	for(const TPair<FString, FPreflightData>& Pair : UnlinkedHordePreflights)
+	{
+		if(Pair.Value.CachedResults.State != EPreflightState::Completed)
+		{
+			WaitTime = Definition.FetchPreflightEachSecondsWhenInProgress;
+			break;
+		}
+	}
+
+	QueueFetch(true, WaitTime);
+}
+
+void FPreflightService::FetchPreflightInfo(bool bRequeue, const FString& InOAuthToken)
+{
+	TSharedPtr<FCredentialsService> CredentialsService = ServiceProvider->GetService<FCredentialsService>();
+	if(Definition.HordeServerAddress.IsEmpty() || !CredentialsService->IsOIDCTokenEnabled())
+	{
+		return;
+	}
+
+	TSharedPtr<FChangelistService> ChangelistService = ServiceProvider->GetService<FChangelistService>();
+
+	// Don't bother with the default changelist
+	if(ChangelistService->GetCLID() == TEXT("default"))
+	{
+		if(bRequeue)
+		{
+			QueueFetch(bRequeue, Definition.FetchPreflightEachSeconds);
+		}
+
+		return;
+	}
+
+	const FString& OIDCToken = CredentialsService->IsTokenReady() ? CredentialsService->GetToken() : InOAuthToken;
+
+	if(!OIDCToken.IsEmpty())
+	{
+		if(!LinkedPFRequest.IsValid())
+		{
+			LinkedPFRequest = FHttpModule::Get().CreateRequest();
+
+			FString FetchPreflightUrl = FString::Format(*Definition.FindPreflightURLFormat, GetFormatParameters());
+			LinkedPFRequest->SetURL(FetchPreflightUrl);
+			LinkedPFRequest->SetVerb(TEXT("GET"));
+		}
+		else if(LinkedPFRequest->GetStatus() == EHttpRequestStatus::Processing)
+		{
+			// if it's still Processing, do not try to request again.
+			return;
+		}
+
+		// ensure the token is the most up to date
+		LinkedPFRequest->SetHeader(TEXT("Authorization"), FString::Printf(TEXT("bearer %s"), *OIDCToken));
+
+		if(!LinkedPFRequest->OnProcessRequestComplete().IsBound())
+		{
+			LinkedPFRequest->OnProcessRequestComplete().Unbind();
+		}
+
+		LinkedPFRequest->OnProcessRequestComplete().BindLambda([this, bRequeue, &OIDCToken](FHttpRequestPtr Request, FHttpResponsePtr HttpResponse, bool bConnectedSuccessfully)
+			{
+				if(!bConnectedSuccessfully)
+				{
+					if(HttpResponse.IsValid())
+					{
+						UE_LOG(LogSubmitTool, Warning, TEXT("Unable to connect to horde. Connection error %d"), HttpResponse->GetResponseCode());
+						UE_LOG(LogSubmitToolDebug, Warning, TEXT("Unable to connect to horde. Connection error\nResponse: %s"), *HttpResponse->GetContentAsString());
+					}
+					else
+					{
+						UE_LOG(LogSubmitTool, Warning, TEXT("Unable to connect to horde. Connection error, no response."));
+					}
+					return;
+				}
+
+				if(HttpResponse.IsValid())
+				{
+					//UE_LOG(LogSubmitToolDebug, Verbose, TEXT("Fetch Preflight Response: %s"), *HttpResponse->GetContentAsString());
+					if(EHttpResponseCodes::IsOk(HttpResponse->GetResponseCode()))
+					{
+						TUniquePtr<FPreflightList> NewHordePreflights = MakeUnique<FPreflightList>();
+						FJsonObjectConverter::JsonObjectStringToUStruct<FPreflightList>(FString::Printf(TEXT("{\"PreflightList\" : %s}"), *HttpResponse->GetContentAsString()), NewHordePreflights.Get());
+						NewHordePreflights->Initialize();
+
+						if(PreflightTag != nullptr && !Definition.FindSinglePreflightURLFormat.IsEmpty())
+						{
+							for(FString PreflightId : PreflightTag->GetValues())
+							{
+								if(PreflightId.Equals(TEXT("skip")) || PreflightId.Equals(TEXT("none")))
+								{
+									continue;
+								}
+
+								if(PreflightId.Contains(TEXT("/")))
+								{
+									int32 SlashIdx;
+									PreflightId.FindLastChar(TCHAR('/'), SlashIdx);
+									PreflightId = PreflightId.RightChop(SlashIdx + 1);
+								}
+
+								PreflightId.TrimStartAndEndInline();
+
+								FRegexPattern Pattern = FRegexPattern(TEXT("(?:[0-9]|[a-f]){24}"), ERegexPatternFlags::CaseInsensitive);
+								FRegexMatcher regex = FRegexMatcher(Pattern, PreflightId);
+								bool match = regex.FindNext();
+								if(match)
+								{
+									const FPreflightData* FoundData = NewHordePreflights->PreflightList.FindByPredicate([&PreflightId](const FPreflightData& InData) { return InData.ID == PreflightId; });
+									if(FoundData == nullptr)
+									{
+										FetchUnlinkedPreflight(PreflightId, bRequeue, OIDCToken);
+									}
+								}
+							}
+						}
+
+						if(!HordePreflights.IsValid() || *NewHordePreflights != *HordePreflights)
+						{
+							UE_LOG(LogSubmitToolDebug, Verbose, TEXT("Newer Preflight information received"));
+
+							// Only log when there's a different number of preflights
+							if(!HordePreflights.IsValid() || HordePreflights->PreflightList.Num() != NewHordePreflights->PreflightList.Num())
+							{
+								UE_LOG(LogSubmitTool, Log, TEXT("Retrieved %d preflights for CL %s"), NewHordePreflights->PreflightList.Num(), *ServiceProvider->GetService<FChangelistService>()->GetCLID());
+							}
+
+							HordePreflights = MoveTemp(NewHordePreflights);
+
+							if(PreflightTag != nullptr)
+							{
+								bool bCLDescriptionModified = false;
+
+								if(HordePreflights->PreflightList.Num() != 0)
+								{
+									FString CurrentTagValue = PreflightTag->GetValuesText();
+
+									if(!bStopAskingTagUpdate && !CurrentTagValue.Contains(HordePreflights->PreflightList[0].ID))
+									{
+										EDialogFactoryResult Result = ShowUpdatePreflightTagDialog();
+										if(Result == EDialogFactoryResult::FirstButton)
+										{
+											// Set the latest one as the tag value
+											ModelInterface->SetTagValues(*PreflightTag, HordePreflights->PreflightList[0].ID);
+											bCLDescriptionModified = true;
+										}
+										else
+										{
+											bStopAskingTagUpdate = true;
+										}
+									}
+								}
+
+								if(bCLDescriptionModified)
+								{
+									ModelInterface->ValidateCLDescription();
+								}
+							}
+						}
+					}
+					else
+					{
+						UE_LOG(LogSubmitTool, Warning, TEXT("Could not retrieve preflights, Http code %d."), HttpResponse->GetResponseCode());
+						UE_LOG(LogSubmitToolDebug, Error, TEXT("Fetch preflight failed. Response %s"), *HttpResponse->GetContentAsString());
+					}
+				}
+				else
+				{
+					UE_LOG(LogSubmitTool, Warning, TEXT("Unable to fetch preflights. Failed with code %d"), HttpResponse->GetResponseCode());
+					UE_LOG(LogSubmitToolDebug, Warning, TEXT("Unable to fetch preflights. Failed with code %d\nResponse: %s"), HttpResponse->GetResponseCode(), *HttpResponse->GetContentAsString());
+				}
+
+				if(ActiveUnlinkedRequests == 0)
+				{
+					if(OnPreflightDataUpdated.IsBound())
+					{
+						OnPreflightDataUpdated.Broadcast(HordePreflights, UnlinkedHordePreflights);
+					}
+
+					if(bRequeue)
+					{
+						Requeue();
+					}
+				}
+			});
+
+		FTimespan TimeSinceLast = FDateTime::UtcNow() - LastRequest;
+		if(bRequeue || TimeSinceLast.GetTotalSeconds() > 3)
+		{
+			LastRequest = FDateTime::UtcNow();
+			UE_LOG(LogSubmitToolDebug, Log, TEXT("Fetching preflights for CL %s. URL: %s"), *ChangelistService->GetCLID(), *LinkedPFRequest->GetURL())
+			LinkedPFRequest->ProcessRequest();
+		}
+	}
+	else
+	{
+		CredentialsService->QueueWorkForToken([this, bRequeue](const FString& InToken)
+			{
+				if(!InToken.IsEmpty())
+				{
+					FetchPreflightInfo(bRequeue, InToken);
+				}
+				else
+				{
+					UE_LOG(LogSubmitTool, Warning, TEXT("Couldn't obtain OAuth token login, communication with Horde is not possible."));
+				}
+			});
+	}
+}
+
+void FPreflightService::FetchUnlinkedPreflight(const FString& InPreflightId, bool bRequeue, const FString& InOAuthToken)
+{
+	if(InOAuthToken.IsEmpty())
+	{
+		return;
+	}
+
+	FHttpRequestPtr& UnlinkedPFRequest = UnlinkedPFRequests.FindOrAdd(InPreflightId);
+
+	if(!UnlinkedPFRequest.IsValid())
+	{
+		UnlinkedPFRequest = FHttpModule::Get().CreateRequest();
+
+		FStringFormatNamedArguments ReplaceStringArgs = GetFormatParameters();
+		ReplaceStringArgs.Add(TEXT("PreflightId"), InPreflightId);
+
+
+		FString FetchPreflightUrl = FString::Format(*Definition.FindSinglePreflightURLFormat, ReplaceStringArgs);
+		UnlinkedPFRequest->SetURL(FetchPreflightUrl);
+		UnlinkedPFRequest->SetVerb(TEXT("GET"));
+	}
+	else if(UnlinkedPFRequest->GetStatus() == EHttpRequestStatus::Processing)
+	{
+		// if it's still Processing, do not try to request again.
+		return;
+	}
+
+	UnlinkedPFRequest->SetHeader(TEXT("Authorization"), FString::Printf(TEXT("bearer %s"), *InOAuthToken));
+
+	if(!UnlinkedPFRequest->OnProcessRequestComplete().IsBound())
+	{
+		UnlinkedPFRequest->OnProcessRequestComplete().Unbind();
+	}
+
+	ActiveUnlinkedRequests++;
+	UnlinkedPFRequest->OnProcessRequestComplete().BindLambda([this, bRequeue, &InOAuthToken, InPreflightId](FHttpRequestPtr Request, FHttpResponsePtr HttpResponse, bool bConnectedSuccessfully)
+		{
+			ActiveUnlinkedRequests--;
+			if(!bConnectedSuccessfully)
+			{
+				if(HttpResponse.IsValid())
+				{
+					UE_LOG(LogSubmitTool, Warning, TEXT("Unable to connect to horde. Connection error %d"), HttpResponse->GetResponseCode());
+					UE_LOG(LogSubmitToolDebug, Warning, TEXT("Unable to connect to horde. Connection error\nResponse: %s"), *HttpResponse->GetContentAsString());
+				}
+				else
+				{
+					UE_LOG(LogSubmitTool, Warning, TEXT("Unable to connect to horde. Connection error, no response."));
+				}
+				return;
+			}
+
+			if(HttpResponse.IsValid())
+			{
+				UE_LOG(LogSubmitToolDebug, Verbose, TEXT("Fetch Single Preflight Response: %s"), *HttpResponse->GetContentAsString());
+				if(EHttpResponseCodes::IsOk(HttpResponse->GetResponseCode()))
+				{
+					FPreflightData ReceivedPreflightInfo;
+					FJsonObjectConverter::JsonObjectStringToUStruct<FPreflightData>(*HttpResponse->GetContentAsString(), &ReceivedPreflightInfo);
+					ReceivedPreflightInfo.RecalculateCachedResults();
+
+					if(!UnlinkedHordePreflights.Contains(InPreflightId) || UnlinkedHordePreflights[InPreflightId] != ReceivedPreflightInfo)
+					{
+						UE_LOG(LogSubmitToolDebug, Verbose, TEXT("Newer %s Preflight information received"), *InPreflightId);
+
+						// Only log when the preflight is new
+						if(!UnlinkedHordePreflights.Contains(InPreflightId))
+						{
+							UE_LOG(LogSubmitTool, Log, TEXT("Retrieved information from preflight %s"), *InPreflightId);
+							UnlinkedHordePreflights.Add(InPreflightId, ReceivedPreflightInfo);
+						}
+						else
+						{
+							UnlinkedHordePreflights[InPreflightId] = MoveTemp(ReceivedPreflightInfo);
+						}
+
+					}
+				}
+				else
+				{
+					UE_LOG(LogSubmitTool, Warning, TEXT("Could not retrieve preflights, Http code %d."), HttpResponse->GetResponseCode());
+					UE_LOG(LogSubmitToolDebug, Error, TEXT("Fetch preflight failed. Response %s"), *HttpResponse->GetContentAsString());
+				}
+			}
+			else
+			{
+				UE_LOG(LogSubmitTool, Warning, TEXT("Unable to fetch preflights. Failed with code %d"), HttpResponse->GetResponseCode());
+				UE_LOG(LogSubmitToolDebug, Warning, TEXT("Unable to fetch preflights. Failed with code %d\nResponse: %s"), HttpResponse->GetResponseCode(), *HttpResponse->GetContentAsString());
+			}
+
+			if(ActiveUnlinkedRequests == 0)
+			{
+				if(OnPreflightDataUpdated.IsBound())
+				{
+					OnPreflightDataUpdated.Broadcast(HordePreflights, UnlinkedHordePreflights);
+				}
+
+				if(bRequeue)
+				{
+					Requeue();
+				}
+			}
+		});
+
+	UnlinkedPFRequest->ProcessRequest();
+}
+
+
+void FPreflightService::StartPreflight()
+{
+	TSharedPtr<ISTSourceControlService> SCCService = ServiceProvider->GetService<ISTSourceControlService>();
+	StreamName = SCCService->GetRootStreamName();
+
+	const TArray<FString>& ShelvedPaths = ServiceProvider->GetService<FChangelistService>()->GetShelvedFilesDepotPaths(true);
+
+	if (ShelvedPaths.IsEmpty())
+	{
+		LastErrorMessage = TEXT("Shelve is empty or it couldn't be retrieved from p4, can't request preflight");
+		State = EPreflightServiceState::Error;
+		return;
+	}
+
+	FString CommonPath = ShelvedPaths[0];
+	const FString& LastPath = ShelvedPaths.Last();
+
+	for(size_t i = 0; i < FMath::Min(CommonPath.Len(), LastPath.Len()); ++i)
+	{
+		if(CommonPath[i] != LastPath[i])
+		{
+			CommonPath = CommonPath.Left(i);
+			break;
+		}
+	}
+
+	if(!CommonPath.Equals(TEXT("//")))
+	{
+		int32 NextSlash = CommonPath.Find(TEXT("/"), ESearchCase::IgnoreCase, ESearchDir::FromStart, 2);
+		const FString Depot = CommonPath.Mid(2, NextSlash - 2);
+		size_t StreamDepth = SCCService->GetDepotStreamLength(Depot);
+
+		for(size_t i = 0; i < StreamDepth; ++i)
+		{
+			NextSlash = CommonPath.Find(TEXT("/"), ESearchCase::IgnoreCase, ESearchDir::FromStart, NextSlash+1);
+		}
+
+		if(NextSlash != -1)
+		{
+			const FSCCStream* FoundStream = SCCService->GetSCCStream(CommonPath.Left(NextSlash));
+
+			if(FoundStream != nullptr)
+			{
+				StreamName = FoundStream->Name;
+			}
+		}
+	}
+
+
+	FString StartPreflightUrl = FString::Format(*Definition.StartPreflightURLFormat, GetFormatParameters());
+
+	// If for some reason, our preflight settings are missing, this will be empty, let's not popup a browser with nothing in it
+	if(!StartPreflightUrl.IsEmpty())
+	{
+		UE_LOG(LogSubmitTool, Log, TEXT("Preflight: Starting preflight with URL: \"%s\""), *StartPreflightUrl);
+
+		FPlatformProcess::LaunchURL(*StartPreflightUrl, nullptr, nullptr);
+		State = EPreflightServiceState::Idle;
+
+		// Do a Fetch in 10 and 30 s to try and capture the triggered preflight
+		QueueFetch(false, 10.f);
+		QueueFetch(false, 30.f);
+	}
+	else
+	{
+		LastErrorMessage = TEXT("Missing INI preflight settings");
+		State = EPreflightServiceState::Error;
+	}
+}
+
+void FPreflightService::OnDeleteShelveOperationComplete(const FSourceControlOperationRef& Operation, ECommandResult::Type Result)
+{
+	if(State == EPreflightServiceState::WaitingForDeleteShelve)
+	{
+		if(Result == ECommandResult::Type::Succeeded)
+		{
+			State = EPreflightServiceState::RequestCreateShelve;
+		}
+		else
+		{
+			LastErrorMessage = TEXT("Unable to delete shelve for preflight");
+			State = EPreflightServiceState::Error;
+		}
+	}
+	else
+	{
+		LastErrorMessage = TEXT("Received delete shelve callback when not waiting for it");
+		State = EPreflightServiceState::Error;
+	}
+}
+
+void FPreflightService::OnCreateShelveOperationComplete(const FSourceControlOperationRef& Operation, ECommandResult::Type Result)
+{
+	if(State == EPreflightServiceState::WaitingForCreateShelve)
+	{
+		if(Result == ECommandResult::Type::Succeeded)
+		{
+			State = EPreflightServiceState::StartPreflight;
+		}
+		else
+		{
+			LastErrorMessage = TEXT("Unable to create shelve for preflight");
+			State = EPreflightServiceState::Error;
+		}
+	}
+	else
+	{
+		LastErrorMessage = TEXT("Received create shelve callback when not waiting for it");
+		State = EPreflightServiceState::Error;
+	}
+}
+
+bool FPreflightService::SelectPreflightTemplate(FPreflightTemplateDefinition& OutTemplate) const
+{
+	bool bFoundTemplate = false;
+	TSharedPtr<FChangelistService> ChangelistService = ServiceProvider->GetService<FChangelistService>();
+
+	// Get the list of files in the changelist
+	const TArray<FString>& FilesInCl = bCheckShelveInstead ? ChangelistService->GetShelvedFilesDepotPaths() : ChangelistService->GetFilesDepotPaths();
+
+	// Loop through each definition to see if the files are in the path then check extension
+	for(const FPreflightTemplateDefinition& Def : Definition.Definitions)
+	{
+		FString RegexPat = FConfiguration::Substitute(Def.RegexPath);
+		FRegexPattern Pattern = FRegexPattern(RegexPat, ERegexPatternFlags::CaseInsensitive);
+		for(const FString& File : FilesInCl)
+		{
+			FRegexMatcher regex = FRegexMatcher(Pattern, File);
+			if(regex.FindNext())
+			{
+				OutTemplate = Def;
+				return true;
+			}
+		}
+	}
+
+	return false;
+}
+
+FString FPreflightService::GetAdditionalTasksString(const FPreflightTemplateDefinition& InTemplate) const
+{
+	TStringBuilder<256> AdditionalTaskStrBuilder;
+	const FString BaseString = TEXT("&id-additional-tasks.");
+	const FString EndString = TEXT("=true");
+	TSharedPtr<FChangelistService> ChangelistService = ServiceProvider->GetService<FChangelistService>();
+
+	const TArray<FString>& FilesInCl = bCheckShelveInstead ? ChangelistService->GetShelvedFilesDepotPaths() : ChangelistService->GetFilesDepotPaths();
+	for(const FPreflightAdditionalTask& AdditionalTask : InTemplate.AdditionalTasks)
+	{
+		FString RegexPat = FConfiguration::Substitute(AdditionalTask.RegexPath);
+		FRegexPattern Pattern = FRegexPattern(RegexPat, ERegexPatternFlags::CaseInsensitive);
+		for(const FString& File : FilesInCl)
+		{
+			FRegexMatcher regex = FRegexMatcher(Pattern, File);
+			if(regex.FindNext())
+			{
+				AdditionalTaskStrBuilder.Append(BaseString);
+				AdditionalTaskStrBuilder.Append(AdditionalTask.TaskId);
+				AdditionalTaskStrBuilder.Append(EndString);
+			}
+		}
+	}
+
+	return AdditionalTaskStrBuilder.ToString();
+}
+
+EDialogFactoryResult FPreflightService::ShowRecreateShelveDialog() const
+{
+	if(ModelInterface->GetMainTab().IsValid() && ModelInterface->GetMainTab().Pin()->GetParentWindow().IsValid())
+	{
+		ModelInterface->GetMainTab().Pin()->GetParentWindow()->DrawAttention(FWindowDrawAttentionParameters());
+	}
+
+	const FText TextTitle = FText::FromString(FString::Printf(TEXT("Preflight CL %s: Recreate shelved files?"), *ServiceProvider->GetService<FChangelistService>()->GetCLID()));
+	const FText TextDescription = FText::FromString(TEXT("There are already shelved files in this changelist.\n\nDo you want to delete and recreate your shelf from the latest changes in your local files for use in the preflight?"));
+
+	return FDialogFactory::ShowDialog(TextTitle, TextDescription, TArray<FString>{ TEXT("Re-shelve files"), TEXT("Use existing Shelve"), TEXT("Cancel") });
+}
+
+EDialogFactoryResult FPreflightService::ShowUpdatePreflightTagDialog() const
+{
+	if(ModelInterface->GetMainTab().IsValid() && ModelInterface->GetMainTab().Pin()->GetParentWindow().IsValid())
+	{
+		ModelInterface->GetMainTab().Pin()->GetParentWindow()->DrawAttention(FWindowDrawAttentionParameters());
+	}
+
+	const FText TextTitle = FText::FromString(FString::Printf(TEXT("Preflight CL %s: Newer preflight available"), *ServiceProvider->GetService<FChangelistService>()->GetCLID()));
+	const FText TextDescription = FText::FromString(FString::Printf(TEXT("There is a newer preflight for this changelist:\n<a id=\"browser\" style=\"Hyperlink\" href=\"%sjob/%s\">%s - %s</>\n\nDo you want to update the #preflight tag?"), *Definition.HordeServerAddress, *HordePreflights->PreflightList[0].ID, *HordePreflights->PreflightList[0].Name, *HordePreflights->PreflightList[0].ID));
+
+	return FDialogFactory::ShowDialog(TextTitle, TextDescription, TArray<FString>{ TEXT("Update Tag Value"), TEXT("Cancel") });
+}
diff --git a/Engine/Source/Programs/SubmitTool/Private/Logic/PreflightService.h b/Engine/Source/Programs/SubmitTool/Private/Logic/PreflightService.h
new file mode 100644
--- /dev/null
+++ b/Engine/Source/Programs/SubmitTool/Private/Logic/PreflightService.h
@@ -0,0 +1,122 @@
+// Copyright Epic Games, Inc. All Rights Reserved.
+
+#pragma once
+
+#include "CoreMinimal.h"
+#include "Containers/Ticker.h"
+#include "ChangelistService.h"
+#include "Logic/ProcessWrapper.h"
+#include "Logic/DialogFactory.h"
+#include "Parameters/SubmitToolParameters.h"
+#include "Models/PreflightData.h"
+#include "Services/Interfaces/ISubmitToolService.h"
+
+class FSubmitToolServiceProvider;
+class FModelInterface;
+class IHttpRequest;
+
+enum class EPreflightServiceState
+{
+	Idle,
+	RequestDeleteShelve,
+	WaitingForDeleteShelve,
+	RequestCreateShelve,
+	WaitingForCreateShelve,
+	StartPreflight,
+	Error,
+};
+
+
+class FPreflightService final : public ISubmitToolService
+{
+public:
+	FPreflightService() = delete;
+	FPreflightService(
+		const FHordeParameters& InSettings,
+		FModelInterface* InModelInterface,
+		TSharedRef<FSubmitToolServiceProvider> InServiceProvider);
+	~FPreflightService();
+
+	void RequestPreflight();
+	void FetchPreflightInfo(bool bRequeue = false, const FString& InOAuthToken = TEXT(""));
+
+	EPreflightServiceState GetState() const
+	{
+		return State;
+	}
+	bool IsRequestInProgress() const
+	{
+		return State != EPreflightServiceState::Idle;
+	}
+	const TUniquePtr<FPreflightList>& GetPreflightData() const
+	{
+		return HordePreflights;
+	}
+	const TMap<FString, FPreflightData>& GetUnlinkedPreflights() const
+	{
+		return UnlinkedHordePreflights;
+	}
+	const FString& GetHordeServerAddress() const
+	{
+		return Definition.HordeServerAddress;
+	}
+
+	FOnPreflightDataUpdated OnPreflightDataUpdated;
+
+private:
+	void QueueFetch(bool bRequeue, float InSeconds);
+
+	void Requeue();
+
+	void StartPreflight();
+	bool SelectPreflightTemplate(FPreflightTemplateDefinition& OutTemplate) const;
+	FString GetAdditionalTasksString(const FPreflightTemplateDefinition& InTemplate) const;
+
+	void FetchUnlinkedPreflight(const FString& InPreflightId, bool bRequeue, const FString& InOAuthToken);
+
+	EDialogFactoryResult ShowRecreateShelveDialog() const;
+	EDialogFactoryResult ShowUpdatePreflightTagDialog() const;
+
+	void OnDeleteShelveOperationComplete(const FSourceControlOperationRef& Operation, ECommandResult::Type Result);
+	void OnCreateShelveOperationComplete(const FSourceControlOperationRef& Operation, ECommandResult::Type Result);
+
+	bool Tick(float DeltaTime);
+
+	// Definitions from the ini
+	const FHordeParameters Definition;
+
+	// services we depend on
+	TSharedRef<FSubmitToolServiceProvider> ServiceProvider;
+	FTag* PreflightTag;
+	FModelInterface* ModelInterface;
+
+	// local data
+	bool bCheckShelveInstead = false;
+	bool bStopAskingTagUpdate = false;
+	FDateTime LastRequest = FDateTime::MinValue();
+	EPreflightServiceState State;
+	FTSTicker::FDelegateHandle TickHandle;
+	FString LastErrorMessage;
+	TUniquePtr<FPreflightList> HordePreflights;
+	TMap<FString, FPreflightData> UnlinkedHordePreflights;
+
+	FString StreamName;
+
+	TMap<FString, FStringFormatArg> GetFormatParameters() const;
+
+	// delete/create shelve callbacks
+	FSourceControlOperationComplete DeleteShelveCallback;
+	FSourceControlOperationComplete CreateShelveCallback;
+
+	// get stream data process
+	TUniquePtr<FProcessWrapper> GetStreamDataProcess;
+	FOnCompleted OnGetStreamDataCompletedCallback;
+	FOnOutputLine OnGetStreamDataOutputLineCallback;
+
+	// Fetch Preflight
+	TSharedPtr<IHttpRequest> LinkedPFRequest = nullptr;
+	TMap<FString, TSharedPtr<IHttpRequest>> UnlinkedPFRequests;
+	int8 ActiveUnlinkedRequests = 0;
+};
+
+Expose_TNameOf(FPreflightService);
\ No newline at end of file
diff --git a/Engine/Source/Programs/SubmitTool/Private/Logic/ProcessWrapper.cpp b/Engine/Source/Programs/SubmitTool/Private/Logic/ProcessWrapper.cpp
new file mode 100644
--- /dev/null
+++ b/Engine/Source/Programs/SubmitTool/Private/Logic/ProcessWrapper.cpp
@@ -0,0 +1,184 @@
+// Copyright Epic Games, Inc. All Rights Reserved.
+
+#include "ProcessWrapper.h"
+
+#include "GenericPlatform/GenericPlatformTime.h"
+#include "HAL/PlatformProcess.h"
+#include "Logging/SubmitToolLog.h"
+#include "Misc/DateTime.h"
+#include "ProcessPipes.h"
+
+FProcessWrapper::FProcessWrapper(const FString& InProcessName, const FString& InPath, const FString& InArgs, const FOnCompleted& InOnCompleted, const FOnOutputLine& InOnOutputLine, const FString& InWorkingDir, const bool InbLaunchHidden, const bool bLaunchReallyHidden, const bool InbLaunchDetached)
+	: ProcessName(InProcessName)
+	, Path(InPath)
+	, Args(InArgs)
+	, WorkingDir(InWorkingDir)
+	, bLaunchesHidden(InbLaunchHidden)
+	, bLaunchesReallyHidden(bLaunchReallyHidden)
+	, bLaunchDetached(InbLaunchDetached)
+	, Pipes(MakeUnique<FProcessPipes>())
+	, OnCompleted(InOnCompleted)
+	, OnOutputLine(InOnOutputLine)
+{}
+
+FProcessWrapper::~FProcessWrapper()
+{
+	ProcessHandle = nullptr;
+
+	FTSTicker::GetCoreTicker().RemoveTicker(TickerHandle);
+	TickerHandle = nullptr;
+}
+
+bool FProcessWrapper::Start(bool bWaitForExit)
+{
+	if(IsRunning())
+	{
+		OutputLine(FString::Format(TEXT("Process %s already running, ignored start request"), { *ProcessName }), EProcessOutputType::ProcessError);
+		return false;
+	}
+	
+	OutputRemainder = FString();
+	ExecutingTime = 0;
+
+	if (!Pipes->Create())
+	{
+		OutputLine(FString::Format(TEXT("Error creating pipes in process {0}"), { *ProcessName }), EProcessOutputType::ProcessError);
+		return false;
+	}
+
+	OutputLine(FString::Format(TEXT("Running process {0}: {1} {2}"), { *ProcessName, *Path, *Args }), EProcessOutputType::ProcessInfo);
+
+	ProcessHandle = MakeUnique<FProcHandle>(FPlatformProcess::CreateProc
+	(
+		*Path,
+		*Args,
+		bLaunchDetached,
+		bLaunchesHidden,
+		bLaunchesReallyHidden,
+		nullptr,
+		0,
+		WorkingDir.Len() ? *WorkingDir : nullptr,
+		Pipes->GetStdOutForProcess(),
+		Pipes->GetStdInForProcess())
+	);
+
+	if (!ProcessHandle->IsValid())
+	{
+		OutputLine(FString::Format(TEXT("Error creating process {0}"), { *ProcessName }), EProcessOutputType::ProcessError);
+		return false;
+	}
+
+	
+	if(bWaitForExit)
+	{
+		FDateTime Before = FDateTime::UtcNow();
+		FPlatformProcess::WaitForProc(*ProcessHandle);
+		FTimespan Timespan = FDateTime::UtcNow() - Before;
+		OnTick(Timespan.GetTotalSeconds());
+		return true;
+	}
+	else
+	{
+		TickerHandle = FTSTicker::GetCoreTicker().AddTicker(FTickerDelegate::CreateRaw(this, &FProcessWrapper::OnTick));
+		return true;
+	}
+}
+
+void FProcessWrapper::Stop()
+{
+	if(!IsRunning())
+	{
+		return;
+	}
+
+	OutputLine(FString::Format(TEXT("Process {0} was stopped"), { *ProcessName }), EProcessOutputType::ProcessInfo);
+
+	FPlatformProcess::TerminateProc(*ProcessHandle, true);
+	
+	Cleanup();
+}
+
+bool FProcessWrapper::IsRunning() const
+{
+	return ProcessHandle.IsValid() && FPlatformProcess::IsProcRunning(*ProcessHandle);
+}
+
+bool FProcessWrapper::OnTick(float Delta)
+{
+	if(!ProcessHandle.IsValid())
+	{
+		return false;
+	}
+
+	if(FPlatformProcess::IsProcRunning(*ProcessHandle))
+	{
+		ReadOutput(false);
+		ExecutingTime += Delta;
+		return true;
+	}
+
+	ExecutingTime += Delta;
+
+	// Flush output
+	ReadOutput(true);	
+	FPlatformProcess::GetProcReturnCode(*ProcessHandle, &ExitCode);
+
+	OutputLine(FString::Printf(TEXT("Completed running process %s. Process took %s and exited with code %d"), *ProcessName, *FGenericPlatformTime::PrettyTime(ExecutingTime), ExitCode), EProcessOutputType::ProcessInfo);
+
+	Cleanup();
+
+	if(OnCompleted.IsBound())
+	{
+		OnCompleted.ExecuteIfBound(ExitCode);
+	}
+
+	return false;
+}
+
+void FProcessWrapper::Cleanup()
+{
+	Pipes->Reset();
+	ProcessHandle = nullptr;
+
+	FTSTicker::GetCoreTicker().RemoveTicker(TickerHandle);
+	TickerHandle = nullptr;
+}
+
+void FProcessWrapper::ReadOutput(bool FlushOutput)
+{
+	if(!OnOutputLine.IsBound())
+	{
+		return;
+	}
+	
+	FString NewOutput = OutputRemainder + FPlatformProcess::ReadPipe(Pipes->GetStdOutForReading());
+
+	if(!NewOutput.IsEmpty() && (NewOutput.Contains(TEXT("\n")) || FlushOutput))
+	{
+		// Find the last line, which may be truncated
+		if(!FlushOutput)
+		{
+			int32 Position;
+			NewOutput.FindLastChar('\n', Position);
+			OutputRemainder = NewOutput.Mid(Position + 1);
+			NewOutput.RemoveFromEnd(OutputRemainder);
+		}
+
+		TArray<FString> Lines;
+		const TCHAR* Separators[] = { TEXT("\n"), TEXT("\r") };
+		NewOutput.ParseIntoArray(Lines, Separators, UE_ARRAY_COUNT(Separators));
+
+		for(const FString& Line : Lines)
+		{
+			OutputLine(Line, EProcessOutputType::SDTOutput);
+		}
+	}
+}
+
+void FProcessWrapper::OutputLine(const FString OutputLine, const EProcessOutputType& OutputType)
+{
+	if(OnOutputLine.IsBound())
+	{
+		OnOutputLine.ExecuteIfBound(OutputLine, OutputType);
+	}
+}
diff --git a/Engine/Source/Programs/SubmitTool/Private/Logic/ProcessWrapper.h b/Engine/Source/Programs/SubmitTool/Private/Logic/ProcessWrapper.h
new file mode 100644
--- /dev/null
+++ b/Engine/Source/Programs/SubmitTool/Private/Logic/ProcessWrapper.h
@@ -0,0 +1,50 @@
+// Copyright Epic Games, Inc. All Rights Reserved.
+
+#pragma once
+#include "Containers/Ticker.h"
+#include "Logic/ProcessPipes.h"
+
+enum EProcessOutputType
+{
+	SDTOutput,
+	ProcessInfo,
+	ProcessError,
+};
+
+DECLARE_DELEGATE_TwoParams(FOnOutputLine, const FString& /*Output*/, const EProcessOutputType&)
+DECLARE_DELEGATE_OneParam(FOnCompleted, const int32 /*Return code*/)
+
+class FProcessWrapper
+{
+public:
+	FProcessWrapper(const FString& InProcessName, const FString& InPath, const FString& InArgs, const FOnCompleted& InOnCompleted = FOnCompleted(), const FOnOutputLine& InOnOutputLine = FOnOutputLine(), const FString& InWorkingDir = FString(), const bool bLaunchHidden = true, const bool bLaunchReallyHidden = true, const bool bLaunchDetached = false);
+	~FProcessWrapper();
+	bool Start(bool bWaitForExit = false);
+	void Stop();
+	bool IsRunning() const;
+
+	int32 ExitCode;
+	float ExecutingTime = 0;
+private:
+	bool OnTick(float Delta);
+	void Cleanup();
+	void ReadOutput(bool FlushOutput);
+	void OutputLine(const FString Output, const EProcessOutputType& OutputType);
+	
+	FString ProcessName;
+	FString Path;
+	FString Args;
+	FString WorkingDir;
+	const bool bLaunchesHidden;
+	const bool bLaunchesReallyHidden;
+	const bool bLaunchDetached;
+
+	TUniquePtr<FProcHandle> ProcessHandle;
+	TUniquePtr<FProcessPipes> Pipes;
+	FString OutputRemainder;
+
+	FTSTicker::FDelegateHandle TickerHandle;
+
+	const FOnCompleted OnCompleted = nullptr;
+	const FOnOutputLine OnOutputLine = nullptr;
+};
diff --git a/Engine/Source/Programs/SubmitTool/Private/Logic/Services/SourceControl/SubmitToolPerforce.cpp b/Engine/Source/Programs/SubmitTool/Private/Logic/Services/SourceControl/SubmitToolPerforce.cpp
new file mode 100644
--- /dev/null
+++ b/Engine/Source/Programs/SubmitTool/Private/Logic/Services/SourceControl/SubmitToolPerforce.cpp
@@ -0,0 +1,952 @@
+// Copyright Epic Games, Inc. All Rights Reserved.
+
+#include "SubmitToolPerforce.h"
+#include "SourceControlInitSettings.h"
+#include "ISourceControlModule.h"
+#include "Logic/ProcessWrapper.h"
+#include "SubmitToolUtils.h"
+#include "Configuration/Configuration.h"
+#include "Logging/SubmitToolLog.h"
+#include "Logic/CredentialsService.h"
+
+#include "CommandLine/CmdLineParameters.h"
+#include "Misc/Base64.h"
+#include "Misc/Paths.h"
+#include "HAL/FileManager.h"
+
+constexpr size_t MaxRecentUsers = 12;
+constexpr size_t RecentUsersDatVersion = 1;
+constexpr size_t MaxRecentGroups = 12;
+constexpr size_t RecentGroupsDatVersion = 1;
+constexpr size_t MaxConnections = 1;
+constexpr size_t MaxConnectionsAttempts = 10;
+
+
+FP4Connection::FP4Connection(FClientApiWrapper& InConnection, FCriticalSection& InMutex)
+	: Connection(InConnection), P4ConnectionMutex(InMutex)
+{
+	FScopeLock Lock(&P4ConnectionMutex);
+	Connection.bIsReady = false;
+}
+
+FP4Connection::~FP4Connection()
+{
+	FScopeLock Lock(&P4ConnectionMutex);
+	Connection.bIsReady = true;
+}
+
+FConnectionPool::~FConnectionPool()
+{
+	for(TUniquePtr<FClientApiWrapper>& ConWrapper : P4Connections)
+	{
+		Error P4Error;
+		ConWrapper->Connection->Final(&P4Error);
+		if(P4Error.Test())
+		{
+			StrBuf ErrorMsg;
+			P4Error.Fmt(&ErrorMsg);
+			UE_LOG(LogSubmitToolP4, Error, TEXT("P4ERROR: Invalid connection to server."));
+			UE_LOG(LogSubmitToolP4, Error, TEXT("%s"), ANSI_TO_TCHAR(ErrorMsg.Text()));
+		}
+	}
+}
+
+TUniquePtr<FP4Connection> FConnectionPool::GetAvailableConnection()
+{
+	FScopeLock Lock(&Mutex);
+	for(TUniquePtr<FClientApiWrapper>& ConWrapper : P4Connections)
+	{
+		if(ConWrapper->bIsReady)
+		{
+			return MakeUnique<FP4Connection>(*ConWrapper, Mutex);
+		}
+	}
+
+
+	if(P4Connections.Num() < MaxConnections)
+	{
+		size_t Idx = CreateConnection();
+		if(Idx != -1)
+		{
+			return MakeUnique<FP4Connection>(*P4Connections[Idx], Mutex);
+		}
+	}
+
+	return nullptr;
+}
+
+size_t FConnectionPool::CreateConnection()
+{
+	if(bConnectionFailed)
+	{
+		return -1;
+	}
+
+	TUniquePtr<ClientApi> P4Client = MakeUnique<ClientApi>();
+	FString Port;
+	FCmdLineParameters::Get().GetValue(FSubmitToolCmdLine::P4Server, Port);
+	P4Client->SetPort(TCHAR_TO_ANSI(*Port));
+
+	Error P4Error;
+	P4Client->Init(&P4Error);
+	if(P4Error.Test())
+	{
+		StrBuf ErrorMsg;
+		P4Error.Fmt(&ErrorMsg);
+		UE_LOG(LogSubmitToolP4, Error, TEXT("P4ERROR: Invalid connection to server."));
+		UE_LOG(LogSubmitToolP4, Error, TEXT("%s"), ANSI_TO_TCHAR(ErrorMsg.Text()));
+		bConnectionFailed = true;
+		return -1;
+	}
+
+	FScopeLock Lock(&Mutex);
+	return P4Connections.Add(MakeUnique<FClientApiWrapper>(MoveTemp(P4Client)));
+}
+
+FSubmitToolPerforce::FSubmitToolPerforce(const FSubmitToolParameters& InParameters)
+: Parameters(InParameters)
+{
+	FModuleManager::LoadModuleChecked<ISourceControlModule>(FName("PerforceSourceControl"));
+
+	FSourceControlInitSettings SCCSettings(FSourceControlInitSettings::EBehavior::OverrideAll);
+
+	const FSourceControlInitSettings::EConfigBehavior IniBehavior = FSourceControlInitSettings::EConfigBehavior::ReadOnly;
+	SCCSettings.SetConfigBehavior(IniBehavior);
+
+	FString PerforceServerAndPort;
+	FCmdLineParameters::Get().GetValue(FSubmitToolCmdLine::P4Server, PerforceServerAndPort);
+
+	FString PerforceUserName;
+	FCmdLineParameters::Get().GetValue(FSubmitToolCmdLine::P4User, PerforceUserName);
+
+	FString PerforceClientName;
+	FCmdLineParameters::Get().GetValue(FSubmitToolCmdLine::P4Client, PerforceClientName);
+
+	SCCSettings.AddSetting(TEXT("P4Port"), PerforceServerAndPort);
+	SCCSettings.AddSetting(TEXT("P4User"), PerforceUserName);
+	SCCSettings.AddSetting(TEXT("P4Client"), PerforceClientName);
+
+	SCCProvider = ISourceControlModule::Get().CreateProvider(FName("Perforce"), TEXT("SubmitTool"), SCCSettings);
+
+	if(SCCProvider.IsValid())
+	{
+		UE_LOG(LogSubmitTool, Log, TEXT("Setting Perforce Connection parameters: %s | User: %s | Workspace: %s"), *PerforceServerAndPort, *PerforceUserName, *PerforceClientName);
+
+		ISourceControlProvider::FInitResult ConnectionResult = SCCProvider->Init(ISourceControlProvider::EInitFlags::AttemptConnection);
+		if(!ConnectionResult.bIsAvailable)
+		{
+			UE_LOG(LogSubmitTool, Error, TEXT("%s"), *ConnectionResult.Errors.ErrorMessage.ToString());
+			for(FText& AdditionalErrorMessage : ConnectionResult.Errors.AdditionalErrors)
+			{
+				UE_LOG(LogSubmitTool, Error, TEXT("%s"), *AdditionalErrorMessage.ToString());
+			}
+		}
+	}
+	else
+	{
+		UE_LOG(LogSubmitTool, Error, TEXT("Failed to create a perforce revision control provider"));
+	}
+
+	TickHandle = FTSTicker::GetCoreTicker().AddTicker(FTickerDelegate::CreateRaw(this, &FSubmitToolPerforce::Tick));
+	GetStream();
+}
+
+FSubmitToolPerforce::~FSubmitToolPerforce()
+{
+	OnUsersGetCallbacks.Clear();
+	OnGroupsGetCallbacks.Clear();
+
+	if(SCCProvider.IsValid())
+	{
+		SCCProvider->Close();
+	}
+
+	FTSTicker::GetCoreTicker().RemoveTicker(TickHandle);
+}
+
+const TUniquePtr<ISourceControlProvider>& FSubmitToolPerforce::GetProvider() const
+{
+	return SCCProvider;
+}
+
+UE::Tasks::TTask<bool> FSubmitToolPerforce::GetStream(const FString& InStream, bool bRequestHierarchy)
+{
+	TArray<FString> Args = { TEXT("-o") };
+
+	if(!InStream.IsEmpty())
+	{
+		Args.Add(InStream);
+	}
+
+	return RunCommand(TEXT("stream"), Args, FOnSCCCommandComplete::CreateLambda(
+		[this, bRequestHierarchy, InStream](bool bSuccess, const FSCCRecordSet& InResultValues, const FSourceControlResultInfo& InResultsInfo)
+		{
+			FScopeLock StreamLock(&StreamMutex);
+			if(bSuccess && InResultValues.Num() > 0 && InResultValues[0].Contains(TEXT("Stream")))
+			{
+				TUniquePtr<FSCCStream> Stream = MakeUnique<FSCCStream>(InResultValues[0][TEXT("Stream")], InResultValues[0][TEXT("Parent")], InResultValues[0][TEXT("Type")]);
+
+				const FString Base = TEXT("Paths");
+				size_t i = 1;
+				FString PathsKey = FString::Printf(TEXT("%s%d"), *Base, i);
+				while(InResultValues[0].Contains(PathsKey))
+				{
+					FString Value = InResultValues[0][PathsKey];
+					if(Value.StartsWith(TEXT("import")))
+					{
+						Value = Value.RightChop(Value.Find(TEXT("//")));
+						while(!Value.EndsWith(TEXT("/")))
+						{
+							Value.RemoveAt(Value.Len() - 1);
+						}
+
+						Stream->AdditionalImportPaths.Add(Value);
+					}
+
+					++i;
+					PathsKey = FString::Printf(TEXT("%s%d"), *Base, i);
+				}
+
+				const FString Parent = Stream->Parent;
+
+				if (!Streams.Contains(Stream->Name))
+				{
+					if (bRequestHierarchy)
+					{
+						StreamHierarchy.Insert(Stream.Get(), 0);
+					}
+
+					Streams.Add(Stream->Name, MoveTemp(Stream));
+
+					if (bRequestHierarchy && !Parent.IsEmpty() && !Parent.Equals(TEXT("none")))
+					{
+						GetStream(Parent);
+					}
+				}
+			}
+			else
+			{
+				TUniquePtr<FSCCStream> Stream = MakeUnique<FSCCStream>(TEXT("Invalid"), FString(), TEXT("Invalid"));
+				if (!InStream.IsEmpty())
+				{
+					if (!Streams.Contains(InStream))
+					{
+						Streams.Add(InStream, MoveTemp(Stream));
+					}
+				}
+				else
+				{
+					if (!Streams.Contains(TEXT("Invalid")))
+					{
+						StreamHierarchy.Add(Stream.Get());
+						Streams.Add(TEXT("Invalid"), MoveTemp(Stream));
+					}
+				}
+			}
+		}
+	));
+}
+
+void FSubmitToolPerforce::GetUsers(const FOnUsersGet::FDelegate& Callback)
+{
+	if(!CachedUsersArray.IsEmpty())
+	{
+		Callback.ExecuteIfBound(CachedUsersArray);
+		return;
+	}
+
+	FScopeLock Lock(&Mutex);
+	OnUsersGetCallbacks.Add(Callback);
+	Lock.Unlock();
+
+	if(!UserTask.IsValid())
+	{
+		UserTask = RunCommand(TEXT("users"), {}, FOnSCCCommandComplete::CreateLambda(
+			[this](bool bSuccess, const FSCCRecordSet& InResultValues, const FSourceControlResultInfo& InResultsInfo)
+			{
+				for(const TMap<FString, FString>& Record : InResultValues)
+				{
+					if(Record.Contains("User"))
+					{
+						TSharedPtr<FUserData> User = MakeShared<FUserData>(Record[TEXT("User")], Record[TEXT("FullName")], Record[TEXT("Email")]);
+						CachedUsersArray.Add(User);
+						CachedUsers.Add(User->Username, User);
+					}
+				}
+
+				LoadRecentUsers();
+
+				FScopeLock Lock(&Mutex);
+				OnUsersGetCallbacks.Broadcast(CachedUsersArray);
+				OnUsersGetCallbacks.Clear();
+				Lock.Unlock();
+			}
+		));
+	}
+}
+
+void FSubmitToolPerforce::LoadRecentUsers()
+{
+	const TUniquePtr<FAES::FAESKey>& Key = FCredentialsService::GetEncryptionKey();
+
+	if(!Key.IsValid())
+	{
+		return;
+	}
+
+	FString RecentUsersString;
+
+	const FString FilePath = GetRecentUsersFilepath();
+	if(IFileManager::Get().FileExists(*FilePath))
+	{
+		FArchive* File = IFileManager::Get().CreateFileReader(*FilePath);
+
+		if(File != nullptr)
+		{
+			int32 Version;
+			*File << Version;
+
+			// Check Versions here
+			if(Version != RecentUsersDatVersion)
+			{
+				UE_LOG(LogSubmitToolDebug, Warning, TEXT("Unexpected Recent Users Version, aborting issues loading."));
+				File->Close();
+				delete File;
+				File = nullptr;
+				return;
+			}
+
+			int32 PaddedLength;
+			int32 LengthWithoutPadding;
+
+			*File << PaddedLength;
+			*File << LengthWithoutPadding;
+
+			TArray<uint8> DeserializedBytes;
+			DeserializedBytes.SetNum(PaddedLength);
+			File->Serialize(DeserializedBytes.GetData(), PaddedLength);
+
+			FAES::DecryptData(DeserializedBytes.GetData(), DeserializedBytes.Num(), *Key);
+
+			RecentUsersString = BytesToString(DeserializedBytes.GetData(), LengthWithoutPadding);
+
+			File->Close();
+			delete File;
+			File = nullptr;
+		}
+		else
+		{
+			UE_LOG(LogSubmitTool, Warning, TEXT("Could not read file '%s'."), *FilePath);
+		}
+	}
+	else
+	{
+		UE_LOG(LogSubmitToolDebug, Log, TEXT("File %s does not exists, no recent users were loaded"), *FilePath);
+	}
+
+	TArray<FString> Usernames;
+	RecentUsersString.ParseIntoArray(Usernames, TEXT(";"));
+
+	RecentUsers.Empty();
+	for(FString& Username : Usernames)
+	{
+		if(const TSharedPtr<FUserData>* User = CachedUsers.Find(Username))
+		{
+			RecentUsers.Add(*User);
+		}
+	}
+}
+
+TSharedPtr<FUserData> FSubmitToolPerforce::GetUserDataFromCache(const FString& Username) const
+{
+	if(CachedUsers.Contains(Username))
+	{
+		return CachedUsers[Username];
+	}
+	return nullptr;
+}
+
+void FSubmitToolPerforce::SaveRecentUsers()
+{
+	const FString RecentUsersString = FString::JoinBy(RecentUsers, TEXT(";"), [](const TSharedPtr<FUserData>& User) { return User->Username; });
+	const TUniquePtr<FAES::FAESKey>& Key = FCredentialsService::GetEncryptionKey();
+
+	const FString FilePath = GetRecentUsersFilepath();
+	FArchive* File = IFileManager::Get().CreateFileWriter(*FilePath, EFileWrite::FILEWRITE_EvenIfReadOnly);
+
+	if(File == nullptr)
+	{
+		UE_LOG(LogSubmitTool, Warning, TEXT("Could not create file '%s'."), *FilePath);
+		return;
+	}
+
+	TArray<uint8> Bytes;
+	Bytes.SetNumUninitialized(RecentUsersString.Len());
+	StringToBytes(RecentUsersString, Bytes.GetData(), RecentUsersString.Len());
+
+	int32 ActualLength = Bytes.Num();
+
+	int32 NumBytesEncrypted = Align(Bytes.Num(), FAES::AESBlockSize);
+	Bytes.SetNum(NumBytesEncrypted);
+	FAES::EncryptData(Bytes.GetData(), Bytes.Num(), *Key);
+
+	int32 Version = RecentUsersDatVersion;
+	*File << Version;
+
+	*File << NumBytesEncrypted;
+	*File << ActualLength;
+
+	File->Serialize(Bytes.GetData(), Bytes.Num());
+
+	File->Close();
+	delete File;
+	File = nullptr;
+}
+
+bool FSubmitToolPerforce::Tick(float InDeltaTime)
+{
+	if(SCCProvider.IsValid())
+	{
+		SCCProvider->Tick();
+	}
+
+	return true;
+}
+
+const TArray<TSharedPtr<FUserData>>& FSubmitToolPerforce::GetRecentUsers() const
+{
+	return RecentUsers;
+}
+
+void FSubmitToolPerforce::AddRecentUser(TSharedPtr<FUserData>& User)
+{
+	if(RecentUsers.Contains(User))
+	{
+		// Remove so we can push the user to the top
+		RecentUsers.Remove(User);
+	}
+
+	if(RecentUsers.Num() >= MaxRecentUsers)
+	{
+		RecentUsers.RemoveAt(MaxRecentUsers - 1);
+	}
+
+	RecentUsers.EmplaceAt(0, User);
+	SaveRecentUsers();
+}
+
+const FString FSubmitToolPerforce::GetRecentUsersFilepath() const
+{
+	return FPaths::Combine(FSubmitToolUtils::GetLocalAppDataPath(), TEXT("SubmitTool"), TEXT("recent_users.dat"));
+}
+
+void FSubmitToolPerforce::GetGroups(const FOnGroupsGet::FDelegate& Callback)
+{
+	if(!CachedGroupsArray.IsEmpty())
+	{
+		Callback.ExecuteIfBound(CachedGroupsArray);
+		return;
+	}
+
+	FScopeLock Lock(&Mutex);
+	OnGroupsGetCallbacks.Add(Callback);
+	Lock.Unlock();
+
+	if(!GroupTask.IsValid())
+	{
+		GroupTask = UE::Tasks::Launch(UE_SOURCE_LOCATION,
+			[this] {
+				TUniquePtr<FP4Connection> Connection = Connections.GetAvailableConnection();
+				
+				size_t Attempts = 0;
+				while(!Connection.IsValid())
+				{
+					if(Attempts >= MaxConnectionsAttempts)
+					{
+						return false;
+					}
+
+					FPlatformProcess::Sleep(1.f);
+					Connection = Connections.GetAvailableConnection();
+
+					++Attempts;
+				}
+
+				ClientApi& P4Client = Connection->GetConnection();
+
+				FString UserName;
+				FCmdLineParameters::Get().GetValue(FSubmitToolCmdLine::P4User, UserName);
+				P4Client.SetUser(TCHAR_TO_ANSI(*UserName));
+
+				FString Client;
+				FCmdLineParameters::Get().GetValue(FSubmitToolCmdLine::P4Client, Client);
+				P4Client.SetClient(TCHAR_TO_ANSI(*Client));
+
+				P4Client.SetProtocol("tag", "");
+
+				FSCCRecordSet ResultValues;
+				FSourceControlResultInfo ResultInfo;
+
+				class FGroupsClientUser : public FSTClientUser
+				{
+					public:
+					FGroupsClientUser(FSCCRecordSet& InRecords, EP4ClientUserFlags InFlags, FSourceControlResultInfo& OutResultInfo, TArray<TSharedPtr<FString>>& OutGroupsArray, const TArray<FString>& InGroupsToExclude)
+					: FSTClientUser(InRecords, InFlags, OutResultInfo), GroupsArray(OutGroupsArray), ExcludedGroups(InGroupsToExclude)
+					{}
+
+					virtual void OutputStat(StrDict* VarList) override
+					{
+						StrRef Var, Value;
+
+						FString Group;
+						bool bIsUser = false;
+						bool bIsSubGroup = false;
+						// Iterate over each variable and add to records
+						for(int32 Index = 0; VarList->GetVar(Index, Var, Value); Index++)
+						{
+							if(Var == "isSubgroup")
+							{
+								bIsSubGroup = Value != "0";
+							}
+
+							if(Var == "group")
+							{
+								Group = TO_TCHAR(Value.Text(), IsUnicodeServer());
+							}
+						}
+
+						if(!Group.IsEmpty() && !GroupsArray.ContainsByPredicate([&Group](const TSharedPtr<FString>& InStr) { return InStr->Equals(Group, ESearchCase::IgnoreCase); }) && !bIsSubGroup)
+						{
+							bool bIsExcluded = false;
+							for(const FString& Filter : ExcludedGroups)
+							{
+								if(Group.StartsWith(Filter))
+								{
+									bIsExcluded = true;
+									break;
+								}
+							}
+
+							if(!bIsExcluded)
+							{
+								GroupsArray.Add(MakeShared<FString>(Group));
+							}
+						}
+					}
+
+					TArray<TSharedPtr<FString>>& GroupsArray;
+					const TArray<FString>& ExcludedGroups;
+				};
+
+				FGroupsClientUser P4User(ResultValues, bIsUnicodeServer ? EP4ClientUserFlags::UnicodeServer : EP4ClientUserFlags::None, ResultInfo, CachedGroupsArray, Parameters.GeneralParameters.GroupsToExclude);
+
+				UE_LOG(LogSubmitToolP4Debug, Log, TEXT("Running command: p4 -p %s -u %s -c %s -ztag groups"), TO_TCHAR(P4Client.GetPort().Text(), bIsUnicodeServer), *UserName, *Client);
+				P4Client.Run(FROM_TCHAR(TEXT("groups"), bIsUnicodeServer), &P4User);
+
+				for(const FText& Msg : ResultInfo.InfoMessages)
+				{
+					UE_LOG(LogSubmitToolP4Debug, Verbose, TEXT("p4 groups: %s"), *Msg.ToString());
+				}
+
+				if(ResultInfo.HasErrors())
+				{
+					for(const FText& Error : ResultInfo.ErrorMessages)
+					{
+						UE_LOG(LogSubmitToolP4, Error, TEXT("p4 groups: %s"), *Error.ToString());
+					}
+				}
+
+				LoadRecentGroups();
+
+				FScopeLock Lock(&Mutex);
+				OnGroupsGetCallbacks.Broadcast(CachedGroupsArray);
+				OnGroupsGetCallbacks.Clear();
+
+				return !ResultInfo.HasErrors();
+
+			});
+	}
+}
+
+void FSubmitToolPerforce::LoadRecentGroups()
+{
+	const TUniquePtr<FAES::FAESKey>& Key = FCredentialsService::GetEncryptionKey();
+
+	if(!Key.IsValid())
+	{
+		return;
+	}
+
+	FString RecentGroupsString;
+
+	const FString FilePath = GetRecentGroupsFilepath();
+	if(IFileManager::Get().FileExists(*FilePath))
+	{
+		FArchive* File = IFileManager::Get().CreateFileReader(*FilePath);
+
+		if(File != nullptr)
+		{
+			int32 Version;
+			*File << Version;
+
+			// Check Versions here
+			if(Version != RecentGroupsDatVersion)
+			{
+				UE_LOG(LogSubmitToolDebug, Warning, TEXT("Unexpected Recent Groups Version, aborting issues loading."));
+				File->Close();
+				delete File;
+				File = nullptr;
+				return;
+			}
+
+			int32 PaddedLength;
+			int32 LengthWithoutPadding;
+
+			*File << PaddedLength;
+			*File << LengthWithoutPadding;
+
+			TArray<uint8> DeserializedBytes;
+			DeserializedBytes.SetNum(PaddedLength);
+			File->Serialize(DeserializedBytes.GetData(), PaddedLength);
+
+			FAES::DecryptData(DeserializedBytes.GetData(), DeserializedBytes.Num(), *Key);
+
+			RecentGroupsString = BytesToString(DeserializedBytes.GetData(), LengthWithoutPadding);
+
+			File->Close();
+			delete File;
+			File = nullptr;
+		}
+		else
+		{
+			UE_LOG(LogSubmitTool, Warning, TEXT("Could not read file '%s'."), *FilePath);
+		}
+	}
+	else
+	{
+		UE_LOG(LogSubmitToolDebug, Log, TEXT("File %s does not exists, no recent groups were loaded"), *FilePath);
+	}
+
+	TArray<FString> GroupNames;
+	RecentGroupsString.ParseIntoArray(GroupNames, TEXT(";"));
+
+	RecentGroups.Empty();
+	for(FString& Name : GroupNames)
+	{
+		if(const TSharedPtr<FString>* Group = CachedGroupsArray.FindByPredicate([&Name](TSharedPtr<FString>& InStr) { return InStr->Equals(Name, ESearchCase::IgnoreCase); }))
+		{
+			RecentGroups.Add(*Group);
+		}
+	}
+}
+
+void FSubmitToolPerforce::SaveRecentGroups()
+{
+	const FString RecentGroupsString = FString::JoinBy(RecentGroups, TEXT(";"), [](const TSharedPtr<FString>& Group) { return *Group; });
+	const TUniquePtr<FAES::FAESKey>& Key = FCredentialsService::GetEncryptionKey();
+
+	const FString FilePath = GetRecentGroupsFilepath();
+	FArchive* File = IFileManager::Get().CreateFileWriter(*FilePath, EFileWrite::FILEWRITE_EvenIfReadOnly);
+
+	if(File == nullptr)
+	{
+		UE_LOG(LogSubmitTool, Warning, TEXT("Could not create file '%s'."), *FilePath);
+		return;
+	}
+
+	TArray<uint8> Bytes;
+	Bytes.SetNumUninitialized(RecentGroupsString.Len());
+	StringToBytes(RecentGroupsString, Bytes.GetData(), RecentGroupsString.Len());
+
+	int32 ActualLength = Bytes.Num();
+
+	int32 NumBytesEncrypted = Align(Bytes.Num(), FAES::AESBlockSize);
+	Bytes.SetNum(NumBytesEncrypted);
+	FAES::EncryptData(Bytes.GetData(), Bytes.Num(), *Key);
+
+	int32 Version = RecentGroupsDatVersion;
+	*File << Version;
+
+	*File << NumBytesEncrypted;
+	*File << ActualLength;
+
+	File->Serialize(Bytes.GetData(), Bytes.Num());
+
+	File->Close();
+	delete File;
+	File = nullptr;
+}
+
+const TArray<TSharedPtr<FString>>& FSubmitToolPerforce::GetRecentGroups() const
+{
+	return RecentGroups;
+}
+
+void FSubmitToolPerforce::AddRecentGroup(TSharedPtr<FString>& Group)
+{
+	if(RecentGroups.Contains(Group))
+	{
+		// Remove so we can push the user to the top
+		RecentGroups.Remove(Group);
+	}
+
+	if(RecentGroups.Num() >= MaxRecentGroups)
+	{
+		RecentGroups.RemoveAt(MaxRecentGroups - 1);
+	}
+
+	RecentGroups.EmplaceAt(0, Group);
+	SaveRecentGroups();
+}
+
+const FString FSubmitToolPerforce::GetRecentGroupsFilepath() const
+{
+	return FPaths::Combine(FSubmitToolUtils::GetLocalAppDataPath(), TEXT("SubmitTool"), TEXT("recent_groups.dat"));
+}
+
+UE::Tasks::TTask<bool> FSubmitToolPerforce::DownloadFiles(const FString& InFilepath, TArray<FSharedBuffer>& OutFileBuffers)
+{
+	return RunCommand(TEXT("print"), { TEXT("-q"), InFilepath }, nullptr, &OutFileBuffers);
+}
+
+void FSubmitToolPerforce::GetUsersAndGroups(const FOnUsersAndGroupsGet::FDelegate& Callback)
+{
+	OnUsersAndGroupsGetCallbacks.Add(Callback);
+
+	auto OnUsersAndGroupsReady = [this]() {
+		if(bUsersReady && bGroupsReady)
+		{
+			OnUsersAndGroupsGetCallbacks.Broadcast(CachedUsersArray, CachedGroupsArray);
+			OnUsersAndGroupsGetCallbacks.Clear();
+
+		}
+	};
+
+	GetUsers(FOnUsersGet::FDelegate::CreateLambda([OnUsersAndGroupsReady, this](TArray<TSharedPtr<FUserData>>&) { bUsersReady = true; OnUsersAndGroupsReady(); }));
+	GetGroups(FOnGroupsGet::FDelegate::CreateLambda([OnUsersAndGroupsReady, this](TArray<TSharedPtr<FString>>&) { bGroupsReady = true; OnUsersAndGroupsReady(); }));
+}
+
+UE::Tasks::TTask<bool> FSubmitToolPerforce::RunCommand(const FString& InCommand, const TArray<FString>& InAdditionalArgs, FOnSCCCommandComplete InCompleteCallback, TArray<FSharedBuffer>* OutData)
+{
+	return UE::Tasks::Launch(UE_SOURCE_LOCATION, [this, InCommand, InAdditionalArgs, InCompleteCallback, OutData] {
+		FSCCRecordSet ResultValues;
+		FSourceControlResultInfo ResultInfo;
+
+		EP4ClientUserFlags Flags = DefaultFlags;
+		Flags |= bIsUnicodeServer ? EP4ClientUserFlags::UnicodeServer : EP4ClientUserFlags::None;
+
+		bool bResult = RunP4Command(InCommand, InAdditionalArgs, ResultValues, ResultInfo, Flags, OutData);
+
+		InCompleteCallback.ExecuteIfBound(bResult, ResultValues, ResultInfo);
+
+		return bResult;
+		});
+}
+
+const TArray<FSCCStream*>& FSubmitToolPerforce::GetClientStreams() const
+{
+	FScopeLock Stream(&StreamMutex);
+	return StreamHierarchy;
+}
+
+const FSCCStream* FSubmitToolPerforce::GetSCCStream(const FString& InStreamName)
+{
+	if(Streams.Contains(InStreamName))
+	{
+		FScopeLock Stream(&StreamMutex);
+		return Streams[InStreamName].Get();
+	}
+	else
+	{
+		if(GetStream(InStreamName, false).GetResult() && Streams.Contains(InStreamName))
+		{
+			FScopeLock Stream(&StreamMutex);
+			return Streams[InStreamName].Get();
+		}
+		else
+		{
+			return nullptr;
+		}
+	}
+}
+
+const FString FSubmitToolPerforce::GetRootStreamName()
+{
+	if(StreamHierarchy.Num() == 0)
+	{
+		GetStream().GetResult();
+	}
+
+	FScopeLock Stream(&StreamMutex);
+	if(StreamHierarchy.Num() > 0)
+	{
+		return StreamHierarchy[0]->Name;
+	}
+
+	return FString();
+}
+
+const FString FSubmitToolPerforce::GetCurrentStreamName()
+{
+	if(StreamHierarchy.Num() == 0)
+	{
+		GetStream().GetResult();
+	}
+
+	FScopeLock Stream(&StreamMutex);
+	if(StreamHierarchy.Num() > 0)
+	{
+		return StreamHierarchy.Last()->Name;
+	}
+
+	return FString();
+}
+
+const size_t FSubmitToolPerforce::GetDepotStreamLength(const FString& InDepotName)
+{
+	if(!DepotStreamLengths.Contains(InDepotName))
+	{
+		RunCommand(TEXT("depot"), { TEXT("-o"), InDepotName }, FOnSCCCommandComplete::CreateLambda(
+			[this, InDepotName](bool bSuccess, const FSCCRecordSet& InResultValues, const FSourceControlResultInfo& InResultsInfo)
+			{
+				if(bSuccess)
+				{
+					if(InResultValues[0].Contains(TEXT("StreamDepth")))
+					{
+						const FString& StreamDepth = InResultValues[0][TEXT("StreamDepth")];
+
+						size_t Depth = 0;
+						for(size_t i = 2; i < StreamDepth.Len(); ++i)
+						{
+							if(StreamDepth[i] == TCHAR('/'))
+							{
+								Depth++;
+							}
+						}
+
+						DepotStreamLengths.Add(InDepotName, Depth);
+					}
+				}
+			})).Wait();
+	}
+
+	return DepotStreamLengths[InDepotName];
+}
+
+const FAuthTicket& FSubmitToolPerforce::GetAuthTicket()
+{
+	if(!P4Ticket.IsValid())
+	{
+		TUniquePtr<FP4Connection> Connection = Connections.GetAvailableConnection();
+		const StrPtr& Username = Connection->GetConnection().GetUser();
+		const StrPtr& Ticket = Connection->GetConnection().GetPassword();
+		P4Ticket = FAuthTicket(TO_TCHAR(Username.Text(), bIsUnicodeServer), TO_TCHAR(Ticket.Text(), bIsUnicodeServer));
+	}
+
+	return P4Ticket;
+}
+
+bool FSubmitToolPerforce::RunP4Command(const FString& InCommand, const TArray<FString>& InAdditionalArgs, TArray<TMap<FString, FString>>& OutResultValues, FSourceControlResultInfo& OutResults, EP4ClientUserFlags InCmdFlags, TArray<FSharedBuffer>* OutData)
+{
+	InCmdFlags |= bIsUnicodeServer ? EP4ClientUserFlags::UnicodeServer : EP4ClientUserFlags::None;
+	InCmdFlags |= OutData != nullptr ? EP4ClientUserFlags::CollectData : EP4ClientUserFlags::None;
+
+	TUniquePtr<FP4Connection> Connection = Connections.GetAvailableConnection();
+
+	size_t Attempts = 0;
+	while(!Connection.IsValid())
+	{
+		if(Attempts >= MaxConnectionsAttempts)
+		{
+			return false;
+		}
+
+		FPlatformProcess::Sleep(1.f);
+		Connection = Connections.GetAvailableConnection();
+
+		++Attempts;
+	}
+
+	ClientApi& P4Client = Connection->GetConnection();
+	FString FullCommand = TEXT("p4 -p ");
+	FullCommand += TO_TCHAR(P4Client.GetPort().Text(), bIsUnicodeServer);
+
+	if(EnumHasAnyFlags(InCmdFlags, EP4ClientUserFlags::UseUser))
+	{
+		FString UserName;
+		FCmdLineParameters::Get().GetValue(FSubmitToolCmdLine::P4User, UserName);
+		P4Client.SetUser(TCHAR_TO_ANSI(*UserName));
+
+		FullCommand += TEXT(" -u ") + UserName;
+	}
+
+	if(EnumHasAnyFlags(InCmdFlags, EP4ClientUserFlags::UseClient))
+	{
+		FString Client;
+		FCmdLineParameters::Get().GetValue(FSubmitToolCmdLine::P4Client, Client);
+		P4Client.SetClient(TCHAR_TO_ANSI(*Client));
+
+		FullCommand += TEXT(" -c ") + Client;
+	}
+
+	if(EnumHasAnyFlags(InCmdFlags, EP4ClientUserFlags::UseZTag))
+	{
+		P4Client.SetProtocol("tag", "");
+		FullCommand += TEXT(" -ztag");
+	}
+
+	FullCommand += TEXT(" ");
+	FullCommand += InCommand;
+
+	int32 ArgC = InAdditionalArgs.Num();
+	UTF8CHAR** ArgV = new UTF8CHAR * [ArgC];
+	for(int32 Index = 0; Index < ArgC; Index++)
+	{
+		if(bIsUnicodeServer)
+		{
+			FTCHARToUTF8 UTF8String(*InAdditionalArgs[Index]);
+			ArgV[Index] = new UTF8CHAR[UTF8String.Length() + 1];
+			FMemory::Memcpy(ArgV[Index], UTF8String.Get(), UTF8String.Length() + 1);
+		}
+		else
+		{
+			ArgV[Index] = new UTF8CHAR[InAdditionalArgs[Index].Len() + 1];
+			FMemory::Memcpy(ArgV[Index], TCHAR_TO_ANSI(*InAdditionalArgs[Index]), InAdditionalArgs[Index].Len() + 1);
+		}
+
+		FullCommand += TEXT(" ") + InAdditionalArgs[Index];
+	}
+
+	P4Client.SetArgv(ArgC, (char**)ArgV);
+	FSTClientUser P4User(OutResultValues, InCmdFlags, OutResults);
+
+	UE_LOG(LogSubmitToolP4Debug, Log, TEXT("Running command: %s"), *FullCommand);
+	P4Client.Run(FROM_TCHAR(*InCommand, bIsUnicodeServer), &P4User);
+
+	for(const FText& Msg : OutResults.InfoMessages)
+	{
+		UE_LOG(LogSubmitToolP4Debug, Verbose, TEXT("%s: %s"), *FullCommand, *Msg.ToString());
+	}
+
+	if(OutResults.HasErrors())
+	{
+		for(const FText& Error : OutResults.ErrorMessages)
+		{
+			UE_LOG(LogSubmitToolP4, Error, TEXT("%s: %s"), *FullCommand, *Error.ToString());
+		}
+	}
+
+	if(OutData != nullptr)
+	{
+		*OutData = P4User.ReleaseData();
+	}
+
+	// Free arguments
+	for(int32 Index = 0; Index < ArgC; Index++)
+	{
+		delete[] ArgV[Index];
+	}
+	delete[] ArgV;
+
+	return !OutResults.HasErrors();
+}
diff --git a/Engine/Source/Programs/SubmitTool/Private/Logic/Services/SourceControl/SubmitToolPerforce.h b/Engine/Source/Programs/SubmitTool/Private/Logic/Services/SourceControl/SubmitToolPerforce.h
new file mode 100644
--- /dev/null
+++ b/Engine/Source/Programs/SubmitTool/Private/Logic/Services/SourceControl/SubmitToolPerforce.h
@@ -0,0 +1,129 @@
+// Copyright Epic Games, Inc. All Rights Reserved.
+
+#pragma once
+
+#include "Logic/Services/Interfaces/ISTSourceControlService.h"
+#include "Parameters/SubmitToolParameters.h"
+#include "PerforceData.h"
+#include "Containers/Ticker.h"
+#include "Containers/LockFreeFixedSizeAllocator.h"
+
+class FProcessWrapper;
+class ISourceControlProvider;
+
+struct FClientApiWrapper
+{
+	FClientApiWrapper(TUniquePtr<ClientApi>&& InConnection)
+	: bIsReady(true), Connection(Forward<TUniquePtr<ClientApi>>(InConnection))
+	{}
+
+	bool bIsReady = true;
+	TUniquePtr<ClientApi> Connection;
+};
+
+struct FP4Connection
+{
+public:
+	FP4Connection(FClientApiWrapper& InConnection, FCriticalSection& InMutex);
+	~FP4Connection();
+
+	ClientApi& GetConnection() const { return *Connection.Connection; }
+
+private:
+	FClientApiWrapper& Connection;
+	FCriticalSection& P4ConnectionMutex;
+};
+
+class FConnectionPool
+{
+public:
+	FConnectionPool() = default;
+	~FConnectionPool();
+	TUniquePtr<FP4Connection> GetAvailableConnection();
+	
+private:
+	size_t CreateConnection();
+	bool bConnectionFailed = false;
+	TArray<TUniquePtr<FClientApiWrapper>> P4Connections;
+	FCriticalSection Mutex;
+	const int32 MaxConnections = 5;
+};
+
+class FSubmitToolPerforce final : public ISTSourceControlService
+{
+public:
+	FSubmitToolPerforce(const FSubmitToolParameters& InParameters);
+	~FSubmitToolPerforce();
+
+	virtual const TUniquePtr<ISourceControlProvider>& GetProvider() const override;
+	virtual void GetUsers(const FOnUsersGet::FDelegate& Callback) override;
+	virtual void GetGroups(const FOnGroupsGet::FDelegate& Callback) override;
+	virtual void GetUsersAndGroups(const FOnUsersAndGroupsGet::FDelegate& Callback) override;
+
+	virtual UE::Tasks::TTask<bool> DownloadFiles(const FString& InFilepath, TArray<FSharedBuffer>& OutFileBuffers) override;
+
+	virtual bool IsAvailable() const override
+	{
+		return SCCProvider.IsValid() && SCCProvider->IsAvailable();
+	}
+
+	virtual bool Tick(float InDeltaTime) override;
+
+	virtual const TArray<TSharedPtr<FUserData>>& GetRecentUsers() const override;
+	virtual void AddRecentUser(TSharedPtr<FUserData>& User) override;
+	virtual const TArray<TSharedPtr<FString>>& GetRecentGroups() const override;
+	virtual void AddRecentGroup(TSharedPtr<FString>& Group) override;
+
+	virtual UE::Tasks::TTask<bool> RunCommand(const FString& InCommand, const TArray<FString>& InAdditionalArgs, FOnSCCCommandComplete InCompleteCallback = nullptr, TArray<FSharedBuffer>* OutData = nullptr) override;
+
+	virtual TSharedPtr<FUserData> GetUserDataFromCache(const FString& Username) const override;
+
+	virtual const TArray<FSCCStream*>& GetClientStreams() const override;
+	virtual const FSCCStream* GetSCCStream(const FString& InStreamName) override;
+	virtual const FString GetRootStreamName() override;
+	virtual const FString GetCurrentStreamName() override;
+	virtual const size_t GetDepotStreamLength(const FString& InDepotName) override;
+	virtual const FAuthTicket& GetAuthTicket() override;
+								  
+private:
+	bool RunP4Command(const FString& InCommand, const TArray<FString>& InAdditionalArgs, TArray<TMap<FString, FString>>& OutResultValues, FSourceControlResultInfo& OutResults, EP4ClientUserFlags InCmdFlags, TArray<FSharedBuffer>* OutData = nullptr);
+
+	const FSubmitToolParameters& Parameters;
+	UE::Tasks::TTask<bool> GetStream(const FString& InStream = FString(), bool bRequestHierarchy = true);
+	TMap<FString, TSharedPtr<FUserData>> CachedUsers;
+	TArray<TSharedPtr<FUserData>> CachedUsersArray;
+	TArray<TSharedPtr<FString>> CachedGroupsArray;
+	UE::Tasks::TTask<bool> UserTask;
+	UE::Tasks::TTask<bool> GroupTask;
+	bool bUsersReady = false;
+	bool bGroupsReady = false;
+	FOnUsersGet OnUsersGetCallbacks;
+	FOnGroupsGet OnGroupsGetCallbacks;
+	FOnUsersAndGroupsGet OnUsersAndGroupsGetCallbacks;
+	
+	FTSTicker::FDelegateHandle TickHandle;
+	TUniquePtr<ISourceControlProvider> SCCProvider;
+
+	void LoadRecentUsers();
+	void SaveRecentUsers();
+	const FString GetRecentUsersFilepath() const;
+	TArray<TSharedPtr<FUserData>> RecentUsers;
+	void LoadRecentGroups();
+	void SaveRecentGroups();
+	const FString GetRecentGroupsFilepath() const;
+	TArray<TSharedPtr<FString>> RecentGroups;
+
+	FCriticalSection Mutex;
+	bool bIsUnicodeServer = false;
+
+	mutable FCriticalSection StreamMutex;
+	TMap<FString, TUniquePtr<FSCCStream>> Streams;
+	TMap<FString, size_t> DepotStreamLengths;
+
+	TArray<FSCCStream*> StreamHierarchy;
+	FAuthTicket P4Ticket;
+	
+	FConnectionPool Connections;
+
+	static constexpr EP4ClientUserFlags DefaultFlags = EP4ClientUserFlags::UseClient | EP4ClientUserFlags::UseUser | EP4ClientUserFlags::UseZTag;
+};
diff --git a/Engine/Source/Programs/SubmitTool/Private/Logic/SwarmService.cpp b/Engine/Source/Programs/SubmitTool/Private/Logic/SwarmService.cpp
new file mode 100644
--- /dev/null
+++ b/Engine/Source/Programs/SubmitTool/Private/Logic/SwarmService.cpp
@@ -0,0 +1,322 @@
+// Copyright Epic Games, Inc. All Rights Reserved.
+
+#include "SwarmService.h"
+
+#include "ISourceControlModule.h"
+#include "ISourceControlProvider.h"
+#include "SourceControlOperations.h"
+#include "Interfaces/IHttpResponse.h"
+#include "Logic/Services/SubmitToolServiceProvider.h"
+#include "Logic/Services/Interfaces/ISTSourceControlService.h"
+#include "Logic/TagService.h"
+#include "Logic/ChangelistService.h"
+#include "Serialization/JsonSerializer.h"
+#include "JsonObjectConverter.h"
+
+#include "Logging/SubmitToolLog.h"
+
+FSwarmService::FSwarmService(TSharedRef<FSubmitToolServiceProvider> InServiceProvider) :
+	ServiceProvider(InServiceProvider)
+{
+	GetSwarmURL();
+	ServiceProvider->GetService<ISTSourceControlService>()->GetUsersAndGroups(FOnUsersAndGroupsGet::FDelegate::CreateLambda([this](TArray<TSharedPtr<FUserData>>& InUsers, TArray<TSharedPtr<FString>>& InGroups) {
+		Users = &InUsers;
+		Groups = &InGroups;
+	}));
+}
+
+void FSwarmService::FetchReview(const OnGetReviewComplete& OnComplete)
+{
+	if (!bCanConnect)
+	{
+		return;
+	}
+
+	FString Changelist = ServiceProvider->GetService<FChangelistService>()->GetCLID();
+
+	if (Changelist.IsEmpty())
+	{
+		return;
+	}
+
+	const TSharedRef<IHttpRequest> HttpRequest = FHttpModule::Get().CreateRequest();
+	HttpRequest->SetHeader(TEXT("Content-Type"), TEXT("application/json; charset=utf-8"));
+	HttpRequest->SetHeader(TEXT("Authorization"), *ServiceProvider->GetService<ISTSourceControlService>()->GetAuthTicket().ToString());
+
+	HttpRequest->SetURL(FString::Format(TEXT("{0}?change={1}&max={2}"),
+		{
+			ReviewsURL(),
+			Changelist, // get the review for a specific CL
+			1 // we only want a single review.
+		}));
+	HttpRequest->SetVerb(TEXT("GET"));
+
+	HttpRequest->OnProcessRequestComplete().BindLambda([OnComplete, this](FHttpRequestPtr Request, FHttpResponsePtr Response, bool bConnectedSuccessfully)
+		{
+			if (!bConnectedSuccessfully)
+			{
+				OnComplete.ExecuteIfBound(Review, TEXT("Connection Failed"));
+				return;
+			}
+
+			UE_LOG(LogSubmitToolDebug, Log, TEXT("Fetch review Response: %s"), *Response->GetContentAsString());
+
+			if(EHttpResponseCodes::IsOk(Response->GetResponseCode()))
+			{
+				FSwarmReviewCollection ReviewCollection;
+				if(FJsonObjectConverter::JsonObjectStringToUStruct<FSwarmReviewCollection>(Response->GetContentAsString(), &ReviewCollection, 0, 0))
+				{
+					if(ReviewCollection.Reviews.Num() > 0)
+					{
+						Review = MakeUnique<FSwarmReview>(ReviewCollection.Reviews[0]);
+						OnComplete.ExecuteIfBound(Review, {});
+					}
+					else
+					{
+						OnComplete.ExecuteIfBound(Review, TEXT("No available reviews."));
+					}
+				}
+				else
+				{
+					OnComplete.ExecuteIfBound(Review, TEXT("Could not parse the response json."));
+				}
+			}
+			else
+			{
+				UE_LOG(LogSubmitTool, Error, TEXT("Could not communicate with swarm due to error %d.\n%s"), Response->GetResponseCode(), *Response->GetContentAsString());
+				OnComplete.ExecuteIfBound(Review, FString::Printf(TEXT("Error code %d."), Response->GetResponseCode()));
+			}
+		});
+
+	HttpRequest->ProcessRequest();
+}
+
+void FSwarmService::CreateReview(const OnCreateReviewComplete& OnComplete)
+{
+	if(!bCanConnect)
+	{
+		OnComplete.ExecuteIfBound(false, FString());
+		return;
+	}
+
+	FString Changelist = ServiceProvider->GetService<FChangelistService>()->GetCLID();
+
+	if(Changelist.IsEmpty())
+	{
+		return;
+	}
+
+	if(CreateSwarmRequest.IsValid())
+	{
+		CreateSwarmRequest->CancelRequest();
+	}
+
+	CreateSwarmRequest = FHttpModule::Get().CreateRequest();
+	CreateSwarmRequest->SetHeader(TEXT("Content-Type"), TEXT("application/json; charset=utf-8"));
+	CreateSwarmRequest->SetHeader(TEXT("Authorization"), *ServiceProvider->GetService<ISTSourceControlService>()->GetAuthTicket().ToString());
+
+	CreateSwarmRequest->SetURL(ReviewsURL());
+	CreateSwarmRequest->SetVerb(TEXT("POST"));
+
+	TSharedRef<FJsonObject> RequestJson = MakeShared<FJsonObject>();
+	RequestJson->SetNumberField(TEXT("change"), FCString::Atoi(*Changelist));
+
+	TArray<TSharedPtr<FJsonValue>> ReviewersObject;
+	TArray<TSharedPtr<FJsonValue>> GroupsObject;
+
+	for (const FTag* Tag : ServiceProvider->GetService<FTagService>()->GetTagsArray())
+	{
+		if (Tag->Definition.InputSubType.Equals(TEXT("Swarm"), ESearchCase::IgnoreCase))
+		{
+			for (const FString& value : Tag->GetValues())
+			{
+				if (Users != nullptr)
+				{
+					TSharedPtr<FUserData>* FoundUser = Users->FindByPredicate([value = value.TrimChar(TCHAR('@'))](TSharedPtr<FUserData>& InUserData) { return InUserData->Username == value; });
+					if (FoundUser != nullptr)
+					{
+						ReviewersObject.Add(MakeShared<FJsonValueString>((*FoundUser)->Username));
+					}
+					continue;
+				}
+
+				if (Groups != nullptr)
+				{
+					TSharedPtr<FString>* FoundGroup = Groups->FindByPredicate([&value](TSharedPtr<FString>& InGroupData) { return *InGroupData == value; });
+					if (FoundGroup != nullptr)
+					{
+						TSharedPtr<FJsonObject> GroupObject = MakeShared<FJsonObject>();
+						TSharedPtr<FJsonObject> GroupDetails = MakeShared<FJsonObject>();
+						GroupDetails->SetStringField(TEXT("required"), TEXT("false"));
+						GroupObject->SetObjectField(*(*FoundGroup), GroupDetails);
+						GroupsObject.Add(MakeShared<FJsonValueObject>(GroupObject));
+					}
+				}
+			}
+		}
+	}
+
+	if (!ReviewersObject.IsEmpty())
+	{
+		RequestJson->SetArrayField(TEXT("reviewers"), ReviewersObject);
+	}
+
+	if(!GroupsObject.IsEmpty())
+	{
+		RequestJson->SetArrayField(TEXT("requiredGroups"), GroupsObject);
+	}
+
+	RequestJson->SetStringField(TEXT("description"), ServiceProvider->GetService<FChangelistService>()->GetCLDescription());
+
+	FString BodyString;
+	TSharedRef<TJsonWriter<TCHAR>> JsonWriter = TJsonWriterFactory<>::Create(&BodyString);
+	FJsonSerializer::Serialize(RequestJson, JsonWriter);
+	UE_LOG(LogSubmitToolDebug, Log, TEXT("Create Swarm request body:\n%s"), *BodyString);
+	CreateSwarmRequest->SetContentAsString(BodyString);
+
+	CreateSwarmRequest->OnProcessRequestComplete().BindLambda([OnComplete, this](FHttpRequestPtr Request, FHttpResponsePtr HttpResponse, bool bConnectedSuccessfully)
+		{
+			if(!bConnectedSuccessfully)
+			{
+				if(HttpResponse.IsValid())
+				{
+					UE_LOG(LogSubmitTool, Warning, TEXT("Unable to create swarm review. Connection error %d - %s."), HttpResponse->GetResponseCode(), *EHttpResponseCodes::GetDescription(static_cast<EHttpResponseCodes::Type>(HttpResponse->GetResponseCode())).ToString());
+					UE_LOG(LogSubmitToolDebug, Warning, TEXT("Unable to create swarm review. Connection error\nResponse: %s"), *HttpResponse->GetContentAsString());
+				}
+				else
+				{
+					UE_LOG(LogSubmitTool, Warning, TEXT("Unable to create swarm review. Connection error, no response."));
+				}
+				OnComplete.ExecuteIfBound(false, FString());
+				return;
+			}
+			if(HttpResponse.IsValid())
+			{
+				UE_LOG(LogSubmitToolDebug, Log, TEXT("Create review Response: %s"), *HttpResponse->GetContentAsString());
+				if(EHttpResponseCodes::IsOk(HttpResponse->GetResponseCode()))
+				{
+					TSharedRef<TJsonReader<TCHAR>> Reader = TJsonReaderFactory<TCHAR>::Create(HttpResponse->GetContentAsString());
+					TSharedPtr<FJsonObject> JsonObj;
+					if(!FJsonSerializer::Deserialize(Reader, JsonObj))
+					{
+						UE_LOG(LogSubmitTool, Error, TEXT("Unable to deserialize swarm create response"));
+						return;
+					}
+
+					FSwarmReview ReviewObj;
+					if (FJsonObjectConverter::JsonObjectToUStruct<FSwarmReview>(JsonObj->GetObjectField(TEXT("review")).ToSharedRef(), &ReviewObj))
+					{
+						Review = MakeUnique<FSwarmReview>(ReviewObj);
+					}
+					FString ReviewId = JsonObj->GetObjectField(TEXT("review"))->GetStringField(TEXT("id"));
+					FString ReviewURL = BuildReviewURL(ReviewId);
+					OnComplete.ExecuteIfBound(true, ReviewURL);
+				}
+				else
+				{
+					UE_LOG(LogSubmitTool, Error, TEXT("Could not create a swarm review due to error %d - %s."), HttpResponse->GetResponseCode(), *EHttpResponseCodes::GetDescription(static_cast<EHttpResponseCodes::Type>(HttpResponse->GetResponseCode())).ToString());
+					OnComplete.ExecuteIfBound(false, FString());
+				}
+			}
+			else
+			{
+				UE_LOG(LogSubmitTool, Warning, TEXT("Unable to create swarm review. Failed with unknown connection error"));
+				OnComplete.ExecuteIfBound(false, FString());
+			}
+		});
+
+	UE_LOG(LogSubmitToolP4, Log, TEXT("Creating swarm review"));
+	CreateSwarmRequest->ProcessRequest();
+}
+
+
+void FSwarmService::UpdateReviewDescription(const TDelegate<void(bool)>& OnComplete, const FString& InDescription)
+{
+	if(!bCanConnect || !Review.IsValid())
+	{
+		UE_LOG(LogSubmitTool, Warning, TEXT("Tried to update swarm review but Swarm API is not available or there is no review for this CL"));
+		OnComplete.ExecuteIfBound(false);
+		return;
+	}
+
+	if(UpdateSwarmRequest.IsValid())
+	{
+		UpdateSwarmRequest->CancelRequest();
+	}
+
+	UpdateSwarmRequest = FHttpModule::Get().CreateRequest();
+	UpdateSwarmRequest->SetHeader(TEXT("Content-Type"), TEXT("application/json; charset=utf-8"));
+	UpdateSwarmRequest->SetHeader(TEXT("Authorization"), *ServiceProvider->GetService<ISTSourceControlService>()->GetAuthTicket().ToString());
+
+	FString URL = ReviewsURL() / FString::FromInt(Review->Id);
+	UpdateSwarmRequest->SetURL(URL);
+	UpdateSwarmRequest->SetVerb(TEXT("PATCH"));
+
+	TSharedRef<FJsonObject> RequestJson = MakeShared<FJsonObject>();
+	RequestJson->SetStringField(TEXT("description"), InDescription);
+
+	FString BodyString;
+	TSharedRef<TJsonWriter<TCHAR>> JsonWriter = TJsonWriterFactory<>::Create(&BodyString);
+	FJsonSerializer::Serialize(RequestJson, JsonWriter);
+	UE_LOG(LogSubmitToolDebug, Log, TEXT("Update Swarm request body:\n%s"), *BodyString);
+	UpdateSwarmRequest->SetContentAsString(BodyString);
+
+	UpdateSwarmRequest->OnProcessRequestComplete().BindLambda([OnComplete, this](FHttpRequestPtr Request, FHttpResponsePtr HttpResponse, bool bConnectedSuccessfully)
+		{
+			if(!bConnectedSuccessfully || !HttpResponse.IsValid())
+			{
+				if(HttpResponse.IsValid())
+				{
+					UE_LOG(LogSubmitTool, Warning, TEXT("Unable to update swarm review. Connection error %d - %s."), HttpResponse->GetResponseCode(), *EHttpResponseCodes::GetDescription(static_cast<EHttpResponseCodes::Type>(HttpResponse->GetResponseCode())).ToString());
+					UE_LOG(LogSubmitToolDebug, Warning, TEXT("Unable to update swarm review. Connection error %d\nResponse: %s"), HttpResponse->GetResponseCode(), *HttpResponse->GetContentAsString());
+				}
+				else
+				{
+					UE_LOG(LogSubmitTool, Warning, TEXT("Unable to update swarm review. Connection error, no response."));
+				}
+
+				OnComplete.ExecuteIfBound(false);
+				return;
+			}
+			else
+			{
+				UE_LOG(LogSubmitToolDebug, Log, TEXT("Update review response: %s"), *HttpResponse->GetContentAsString());
+				if(EHttpResponseCodes::IsOk(HttpResponse->GetResponseCode()))
+				{
+					UE_LOG(LogSubmitTool, Log, TEXT("Swarm description updated successfully"));
+					OnComplete.ExecuteIfBound(true);
+				}
+				else
+				{
+					UE_LOG(LogSubmitTool, Error, TEXT("Could not update swarm description due to error %d - %s."), HttpResponse->GetResponseCode(), *EHttpResponseCodes::GetDescription(static_cast<EHttpResponseCodes::Type>(HttpResponse->GetResponseCode())).ToString());
+					OnComplete.ExecuteIfBound(false);
+				}
+			}
+		});
+
+	UE_LOG(LogSubmitToolP4, Log, TEXT("Updating swarm review description"));
+	UpdateSwarmRequest->ProcessRequest();
+
+}
+
+const FString& FSwarmService::GetSwarmURL()
+{
+	if(SwarmURL.IsEmpty())
+	{
+		ServiceProvider->GetService<ISTSourceControlService>()->RunCommand(TEXT("property"), { "-l", "-n", "P4.Swarm.URL" }, FOnSCCCommandComplete::CreateLambda(
+			[this](bool bSuccess, const FSCCRecordSet& InResultValues, const FSourceControlResultInfo& InResultsInfo)
+			{
+				if(bSuccess && InResultValues.Num() > 0 && InResultValues[0].Contains(TEXT("value")))
+				{
+					SwarmURL = InResultValues[0][TEXT("value")];
+				}
+
+			})).Wait();
+	}
+	return SwarmURL;
+}
+
+FString FSwarmService::ReviewsURL() const
+{
+	return SwarmURL.IsEmpty() ? FString{} : SwarmURL / TEXT("api/v9/reviews");
+}
\ No newline at end of file
diff --git a/Engine/Source/Programs/SubmitTool/Private/Logic/TagService.h b/Engine/Source/Programs/SubmitTool/Private/Logic/TagService.h
new file mode 100644
--- /dev/null
+++ b/Engine/Source/Programs/SubmitTool/Private/Logic/TagService.h
@@ -0,0 +1,54 @@
+// Copyright Epic Games, Inc. All Rights Reserved.
+
+#pragma once
+
+#include "CoreMinimal.h"
+#include "Models/Tag.h"
+#include "Parameters/SubmitToolParameters.h"
+#include "ChangelistService.h"
+#include "Services/Interfaces/ISubmitToolService.h"
+
+class FTagService : public ISubmitToolService
+{
+public:
+	FTagService(const FSubmitToolParameters& InParameters, TSharedPtr<FChangelistService> clService);
+
+	void ParseCLDescription();
+
+	void ApplyTag(const FString& tagID);
+	void ApplyTag(FTag& tag);
+	void UpdateTagsInCL();
+	void RemoveTag(const FString& tagID);
+	void RemoveTag(FTag& tag);
+	void SetTagValues(const FString& tagID, const FString& values);
+	void SetTagValues(FTag& tag, const FString& values);
+	void SetTagValues(FTag& tag, const TArray<FString>& values);
+
+	FTag* GetTag(const FString& tagID)
+	{
+		if(RegisteredTags.Contains(tagID))
+		{
+			return &RegisteredTags[tagID];
+		}
+
+		return nullptr;
+	};
+
+	FTagUpdated OnTagUpdated;
+
+	const TArray<const FTag*>& GetTagsArray() const;
+
+private:
+	FString& GetCLDescription() { return const_cast<FString&>(ChangelistService->GetCLDescription()); }
+
+	mutable TArray<const FTag*> CachedTags;
+
+	const FSubmitToolParameters& Parameters;
+	TMap<FString, FTag> RegisteredTags;
+	TSharedPtr<FChangelistService> ChangelistService;
+
+	void RegisterTags();
+	void UpdateTagsPositions(size_t changePos, int32 delta);
+};
+
+Expose_TNameOf(FTagService);
\ No newline at end of file
diff --git a/Engine/Source/Programs/SubmitTool/Private/Logic/UpdateService.cpp b/Engine/Source/Programs/SubmitTool/Private/Logic/UpdateService.cpp
new file mode 100644
--- /dev/null
+++ b/Engine/Source/Programs/SubmitTool/Private/Logic/UpdateService.cpp
@@ -0,0 +1,463 @@
+// Copyright Epic Games, Inc. All Rights Reserved.
+
+#include "UpdateService.h"
+
+#include "HAL/FileManager.h"
+#include "HttpModule.h"
+#include "Interfaces/IHttpResponse.h"
+#include "JsonObjectConverter.h"
+#include "Misc/FileHelper.h"
+#include "Misc/Paths.h"
+
+#include "CommandLine/CmdLineParameters.h"
+#include "Logging/SubmitToolLog.h"
+#include "Logic/ProcessWrapper.h"
+#include "Models/HordeDeploymentData.h"
+#include "Framework/Application/SlateApplication.h"
+
+FUpdateService::FUpdateService(const FHordeParameters& InHordeParameters, const FAutoUpdateParameters& InAutoUpdateParameters, TWeakPtr<FSubmitToolServiceProvider> InServiceProvider) :
+	HordeParameters(InHordeParameters),
+	AutoUpdateParameters(InAutoUpdateParameters),
+	ServiceProvider(InServiceProvider),
+	DownloadFile(nullptr),
+	Downloaded(0)
+{
+}
+
+FUpdateService::~FUpdateService()
+{
+}
+
+bool FUpdateService::CheckForNewVersion()
+{
+	if (!FPaths::IsStaged())
+	{
+		return false;
+	}
+
+	if (!this->AutoUpdateParameters.bIsAutoUpdateOn)
+	{
+		return false;
+	}
+
+	GetLocalVersion();
+	GetLatestVersion();
+
+	UE_LOG(LogSubmitToolDebug, Log, TEXT("Submit Tool Versions:\nLocal Version: %s\nRemote Version: %s"), *LocalVersion, *RemoteVersion);
+
+	if (RemoteVersion.IsEmpty())
+	{
+		return false;
+	}
+	else if(LocalVersion.IsEmpty())
+	{
+		LocalVersion = RemoteVersion;
+		SaveLocalVersionToFile();
+	}
+
+	return LocalVersion != RemoteVersion;
+}
+
+void FUpdateService::StartAutoUpdateScript()
+{
+	// security to avoid deleting the current folder while debugging
+	// the update script assumes that the submit tool is staged and installed in isolation
+	if (!FPaths::IsStaged())
+	{
+		return;
+	}
+	
+	FString Name = TEXT("Submit Tool Auto Update");
+	FString Cmd = AutoUpdateParameters.AutoUpdateCommand;
+	
+#if PLATFORM_MAC
+	// install path is five levels above the engine path for Mac: ExtractTargetFolder/Mac/SubmitTool.app/Contents/UE/Engine
+	const FString InstallFolder = FPaths::ConvertRelativePathToFull(FString::Printf(TEXT("%s/../../../../.."), *FPaths::EngineDir()));
+#else
+	// install path is two levels above the engine path for Linux & Windows: ExtractTargetFolder/Windows/Engine
+	const FString InstallFolder = FPaths::ConvertRelativePathToFull(FString::Printf(TEXT("%s/../.."), *FPaths::EngineDir()));
+#endif
+	// copy the auto update script
+	const FString& UpdateScriptDirectory = FPaths::GetPath(this->AutoUpdateParameters.LocalAutoUpdateScript);
+	IFileManager::Get().MakeDirectory(*UpdateScriptDirectory, true);
+	IFileManager::Get().Copy(*this->AutoUpdateParameters.LocalAutoUpdateScript, *this->AutoUpdateParameters.AutoUpdateScript, true, true);
+
+	FString RootDir;
+	FCmdLineParameters::Get().GetValue(FSubmitToolCmdLine::RootDir, RootDir);
+
+	FString Args = AutoUpdateParameters.AutoUpdateArgs;
+	Args = Args.Replace(TEXT("$(script)"), *FPaths::ConvertRelativePathToFull(AutoUpdateParameters.LocalAutoUpdateScript));
+	Args = Args.Replace(TEXT("$(zip)"), *AutoUpdateParameters.LocalDownloadZip);
+	Args = Args.Replace(TEXT("$(folder)"), *InstallFolder);
+	Args = Args.Replace(TEXT("$(executablename)"), FPlatformProcess::ExecutableName(true));
+	Args = Args.Replace(TEXT("$(version)"), *LatestVersionDownloaded);
+	Args = Args.Replace(TEXT("$(versionfile)"), *AutoUpdateParameters.LocalVersionFile);
+	Args = Args.Replace(TEXT("$(executablepath)"), FPlatformProcess::ExecutablePath());
+	Args = Args.Replace(TEXT("$(rootdir)"), *RootDir.TrimStartAndEnd());
+	Args = Args.Replace(TEXT("$(executableargs)"), *GetSubmitToolArgs());
+
+	UE_LOG(LogSubmitTool, Warning, TEXT("Starting Auto-Update Script %s %s"), *Cmd, *Args);
+
+	const bool bLaunchDetached = false;
+	const bool bLaunchesHidden = false;
+	const bool bLaunchesReallyHidden = false;
+
+	FProcessWrapper ScriptProcess = 
+		FProcessWrapper(TEXT("AutoUpdateScript"),
+			Cmd,
+			Args,
+			nullptr,
+			FOnOutputLine::CreateLambda([](const FString& InOutput, const EProcessOutputType& OutputType){ UE_LOG(LogSubmitToolDebug, Log, TEXT("AutoUpdateOutput: %s"), *InOutput); }),
+			FString(),
+			bLaunchesHidden,
+			bLaunchesReallyHidden,
+			bLaunchDetached);
+
+	ScriptProcess.Start();
+
+	FSlateApplication::Get().CloseAllWindowsImmediately();
+}
+
+const FString& FUpdateService::GetDeployId()
+{
+	if(DeployId.IsEmpty())
+	{
+		UE_LOG(LogSubmitToolDebug, Log, TEXT("Trying to load deploy id from file %s"), *AutoUpdateParameters.DeployIdFilePath);
+		if(FPaths::FileExists(this->AutoUpdateParameters.DeployIdFilePath))
+		{
+			FFileHelper::LoadFileToString(DeployId, *this->AutoUpdateParameters.DeployIdFilePath);
+			DeployId.TrimStartAndEndInline();
+		}
+	}
+
+	return DeployId;
+}
+
+void FUpdateService::Cancel()
+{
+	if (this->DownloadRequest.IsValid() && DownloadRequest->GetStatus() == EHttpRequestStatus::Processing)
+	{
+		this->DownloadRequest->CancelRequest();
+	}
+
+	this->Downloaded = 0;
+}
+
+
+void FUpdateService::InstallLatestVersion()
+{
+	if(GetDeployId().IsEmpty())
+	{
+		return;		
+	}
+
+	TSharedPtr<IHttpRequest> Request = FHttpModule::Get().CreateRequest();
+
+	const FString& Url = FString::Format(TEXT("{0}api/v1/tools/{1}"), { HordeParameters.HordeServerAddress,  DeployId });
+
+	Request->SetURL(Url);
+	Request->SetVerb(TEXT("GET"));
+
+	Request->OnProcessRequestComplete().BindLambda([this](FHttpRequestPtr Request, FHttpResponsePtr HttpResponse, bool bConnectedSuccessfully)
+		{
+			if (!bConnectedSuccessfully)
+			{
+				if (HttpResponse.IsValid())
+				{
+					UE_LOG(LogSubmitTool, Warning, TEXT("Unable to connect to horde. Connection error %d"), HttpResponse->GetResponseCode());
+					UE_LOG(LogSubmitToolDebug, Warning, TEXT("Unable to connect to horde. Connection error\nResponse: %s"), *HttpResponse->GetContentAsString());
+				}
+				else
+				{
+					UE_LOG(LogSubmitTool, Warning, TEXT("Unable to connect to horde. Connection error, no response."));
+				}
+				return;
+			}
+
+			if (HttpResponse.IsValid())
+			{
+				if (EHttpResponseCodes::IsOk(HttpResponse->GetResponseCode()))
+				{
+					FDeploymentList DeploymentList;
+
+					FString Content = HttpResponse->GetContentAsString();
+					FJsonObjectConverter::JsonObjectStringToUStruct<FDeploymentList>(*Content, &DeploymentList);
+
+					const int32 Num = DeploymentList.Deployments.Num();
+					if (Num == 0)
+					{
+						UE_LOG(LogSubmitTool, Warning, TEXT("Unable to retrieve latest deployment from Horde."));
+						return;
+					}
+
+					int32 Latest = 0;
+					for (int Idx = 1; Idx < Num; Idx++)
+					{
+						FDateTime LatestTime;
+						FDateTime::ParseIso8601(*DeploymentList.Deployments[Latest].StartedAt, LatestTime);
+
+						FDateTime CurrentTime;
+						FDateTime::ParseIso8601(*DeploymentList.Deployments[Idx].StartedAt, CurrentTime);
+
+						if (CurrentTime > LatestTime)
+						{
+							Latest = Idx;
+						}
+					}
+
+					const FString& LatestVersion = DeploymentList.Deployments[Latest].Id;
+
+					UE_LOG(LogSubmitTool, Display, TEXT("Local SubmitTool version: %s."), *GetLocalVersion());
+					UE_LOG(LogSubmitTool, Display, TEXT("Latest SubmitTool version available on Horde: %s."), *LatestVersion);
+
+					if (LatestVersion != GetLocalVersion())
+					{
+						UE_LOG(LogSubmitTool, Display, TEXT("Submit Tool needs to be updated."));
+						this->DownloadLatestVersion(DeployId, LatestVersion);
+					}
+				}
+				else
+				{
+					UE_LOG(LogSubmitTool, Warning, TEXT("Unable to fetch latest Submmit Tool deployment. Failed with code %d"), HttpResponse->GetResponseCode());
+					UE_LOG(LogSubmitToolDebug, Warning, TEXT("Unable to fetch latest Submmit Tool deployment. Failed with code %d\nResponse: %s"), HttpResponse->GetResponseCode(), *HttpResponse->GetContentAsString());
+				}
+			}
+		});
+
+	Request->ProcessRequest();
+}
+
+const FString& FUpdateService::GetLatestVersion(bool bForce)
+{
+	if(GetDeployId().IsEmpty())
+	{
+		return RemoteVersion;
+	}
+
+	if(!RemoteVersion.IsEmpty() && !bForce)
+	{
+		return RemoteVersion;
+	}
+
+	TSharedPtr<IHttpRequest> Request = FHttpModule::Get().CreateRequest();
+
+	const FString& Url = FString::Format(TEXT("{0}api/v1/tools/{1}"), { HordeParameters.HordeServerAddress,  DeployId });
+
+	UE_LOG(LogSubmitToolDebug, Log, TEXT("Fetching last version from horde using URL: %s"), *Url);
+
+	Request->SetURL(Url);
+	Request->SetVerb(TEXT("GET"));
+	Request->ProcessRequestUntilComplete();
+
+	FHttpResponsePtr HttpResponse = Request->GetResponse();
+
+	if (Request->GetStatus() == EHttpRequestStatus::Succeeded)
+	{
+		FDeploymentList DeploymentList;
+
+		FString Content = HttpResponse->GetContentAsString();
+		FJsonObjectConverter::JsonObjectStringToUStruct<FDeploymentList>(*Content, &DeploymentList);
+
+		const int32 Num = DeploymentList.Deployments.Num();
+		if (Num == 0)
+		{
+			UE_LOG(LogSubmitTool, Warning, TEXT("Unable to retrieve latest deployment from Horde."));
+			return RemoteVersion;
+		}
+
+		int32 Latest = 0;
+		for (int Idx = 1; Idx < Num; Idx++)
+		{
+			FDateTime LatestTime;
+			FDateTime::ParseIso8601(*DeploymentList.Deployments[Latest].StartedAt, LatestTime);
+
+			FDateTime CurrentTime;
+			FDateTime::ParseIso8601(*DeploymentList.Deployments[Idx].StartedAt, CurrentTime);
+
+			if (CurrentTime > LatestTime)
+			{
+				Latest = Idx;
+			}
+		}
+
+		RemoteVersion = DeploymentList.Deployments[Latest].Id;
+	}
+	else if (Request->GetStatus() == EHttpRequestStatus::Failed)
+	{
+		if (HttpResponse.IsValid())
+		{
+			UE_LOG(LogSubmitTool, Warning, TEXT("Unable to fetch latest Submmit Tool deployment. Failed with code %d"), HttpResponse->GetResponseCode());
+			UE_LOG(LogSubmitToolDebug, Warning, TEXT("Unable to fetch latest Submmit Tool deployment. Failed with code %d\nResponse: %s"), HttpResponse->GetResponseCode(), *HttpResponse->GetContentAsString());
+		}
+		else
+		{
+			UE_LOG(LogSubmitTool, Warning, TEXT("Unable to fetch latest Submmit Tool deployment."));
+		}
+	}
+
+	return RemoteVersion;
+}
+
+const FString& FUpdateService::GetLocalVersion()
+{
+	if (LocalVersion.IsEmpty() && FPaths::FileExists(this->AutoUpdateParameters.LocalVersionFile))
+	{
+		FFileHelper::LoadFileToString(LocalVersion, *this->AutoUpdateParameters.LocalVersionFile);
+			
+		LocalVersion = LocalVersion
+			.TrimStartAndEnd()
+			.Replace(*FString("\r"), *FString(""))
+			.Replace(*FString("\n"), *FString(""))
+			;
+	}
+
+	return LocalVersion;
+}
+
+
+const FString FUpdateService::GetDownloadMessage()
+{
+	if(Downloaded != 0 && DownloadErrorMessage.IsEmpty())
+	{
+		return FString::Printf(TEXT("Downloading: %s"), *GetReadableDownloadSize());
+	}
+	else
+	{
+		return DownloadErrorMessage;
+	}
+}
+
+void FUpdateService::SaveLocalVersionToFile() const
+{
+	if(!LocalVersion.IsEmpty())
+	{
+		UE_LOG(LogSubmitToolDebug, Log, TEXT("Saving current version to file %s"), *AutoUpdateParameters.LocalVersionFile);
+		FFileHelper::SaveStringToFile(LocalVersion, *this->AutoUpdateParameters.LocalVersionFile);
+	}
+}
+
+
+bool FUpdateService::DownloadLatestVersion(const FString& InDeployId, const FString& InLatestVersion)
+{
+	UE_LOG(LogSubmitTool, Display, TEXT("Downloading latest Submit Tool version: %s."), *InLatestVersion);
+
+	FString Url = FString::Format(TEXT("{0}api/v1/tools/{1}?action=download"), { HordeParameters.HordeServerAddress,  InDeployId });
+
+	this->DownloadRequest = FHttpModule::Get().CreateRequest();
+	this->DownloadRequest->SetURL(Url);
+	this->DownloadRequest->SetVerb(TEXT("GET"));
+
+	IFileManager::Get().Delete(*this->AutoUpdateParameters.LocalDownloadZip, false, true, true);
+
+	this->DownloadFile = IFileManager::Get().CreateFileWriter(*this->AutoUpdateParameters.LocalDownloadZip, EFileWrite::FILEWRITE_AllowRead);
+
+	FHttpRequestStreamDelegateV2 StreamDelegate;
+	StreamDelegate.BindLambda([this](void* InDataPtr, int64& InOutLength) {this->OnProcessDownloadRequestStream(InDataPtr, InOutLength); });
+	this->DownloadRequest->SetResponseBodyReceiveStreamDelegateV2(StreamDelegate);
+
+	this->DownloadRequest->OnProcessRequestComplete().BindLambda([this, InLatestVersion](FHttpRequestPtr Request, FHttpResponsePtr HttpResponse, bool bConnectedSuccessfully)
+		{
+			if (!bConnectedSuccessfully)
+			{
+				if (HttpResponse.IsValid())
+				{
+					DownloadErrorMessage = FString::Printf(TEXT("Unable to download latest Submit Tool deployment. Failed with code %d, See Logs for more info."), HttpResponse->GetResponseCode());
+					UE_LOG(LogSubmitTool, Warning, TEXT("Unable to connect to horde. Connection error %d\nResponse: %s"), HttpResponse->GetResponseCode(), *HttpResponse->GetContentAsString());
+				}
+				else
+				{
+					DownloadErrorMessage = FString::Printf(TEXT("Unable to download latest Submit Tool deployment. Unknown Error."), HttpResponse->GetResponseCode());
+					UE_LOG(LogSubmitTool, Warning, TEXT("Unable to connect to horde. Connection error, no response."));
+				}
+				return;
+			}
+
+			if (HttpResponse.IsValid())
+			{
+				if(DownloadFile)
+				{
+					this->DownloadFile->Flush();
+					this->DownloadFile->Close();
+					delete DownloadFile; DownloadFile = nullptr;
+				}
+
+				if (EHttpResponseCodes::IsOk(HttpResponse->GetResponseCode()))
+				{
+					LatestVersionDownloaded = InLatestVersion;
+					this->StartAutoUpdateScript();
+				}
+				else
+				{
+					DownloadErrorMessage = FString::Printf(TEXT("Unable to download latest Submit Tool deployment. Failed with code %d, See Logs for more info."), HttpResponse->GetResponseCode());
+					UE_LOG(LogSubmitTool, Warning, TEXT("Unable to download latest Submit Tool deployment. Failed with code %d"), HttpResponse->GetResponseCode());
+					UE_LOG(LogSubmitToolDebug, Warning, TEXT("Unable to download latest Submit Tool deployment. Failed with code %d\nResponse: %s"), HttpResponse->GetResponseCode(), *HttpResponse->GetContentAsString());
+				}
+			}
+			else
+			{
+				DownloadErrorMessage = FString::Printf(TEXT("Unable to download latest Submit Tool deployment. Unknown Error."), HttpResponse->GetResponseCode());
+				UE_LOG(LogSubmitTool, Warning, TEXT("Unable to connect to horde. Connection error, no response."));
+			}
+		});
+
+	DownloadErrorMessage = FString();
+	DownloadRequest->ProcessRequest();
+	return true;
+}
+
+void FUpdateService::OnProcessDownloadRequestStream(void* InDataPtr, int64& InOutLength)
+{
+	if (DownloadFile == nullptr || !DownloadRequest.IsValid())
+	{
+		return;
+	}
+
+	if (DownloadRequest->GetStatus() == EHttpRequestStatus::Failed)
+	{
+		return;
+	}
+	this->DownloadFile->Serialize(InDataPtr, InOutLength);
+	Downloaded += InOutLength;
+}
+
+FString FUpdateService::GetSubmitToolArgs() const
+{
+
+	FString Changelist;
+	FCmdLineParameters::Get().GetValue(FSubmitToolCmdLine::P4ChangeList, Changelist);
+
+	FString PerforceServerAndPort;
+	FCmdLineParameters::Get().GetValue(FSubmitToolCmdLine::P4Server, PerforceServerAndPort);
+
+	FString PerforceUserName;
+	FCmdLineParameters::Get().GetValue(FSubmitToolCmdLine::P4User, PerforceUserName);
+
+	FString PerforceClientName;
+	FCmdLineParameters::Get().GetValue(FSubmitToolCmdLine::P4Client, PerforceClientName);
+
+	FString ParameterFile;
+	FCmdLineParameters::Get().GetValue(FSubmitToolCmdLine::ParameterFile, ParameterFile);
+
+	FString ExecutableArguments = "-server $(server) -user $(user) -client $(client) -cl $(changelist)";
+	ExecutableArguments = ExecutableArguments.Replace(TEXT("$(server)"), *PerforceServerAndPort);
+	ExecutableArguments = ExecutableArguments.Replace(TEXT("$(client)"), *PerforceClientName);
+	ExecutableArguments = ExecutableArguments.Replace(TEXT("$(user)"), *PerforceUserName);
+	ExecutableArguments = ExecutableArguments.Replace(TEXT("$(changelist)"), *Changelist);
+	ExecutableArguments = ExecutableArguments.Replace(TEXT("$(parameterfile)"), *ParameterFile);
+
+	return ExecutableArguments;
+}
+
+const FString FUpdateService::GetReadableDownloadSize()
+{
+	static const TCHAR* Units[] = { TEXT("B"), TEXT("KB"), TEXT("MB"), TEXT("GB"), TEXT("TB"), TEXT("PB"), TEXT("EB") };
+
+	uint64 UnitIndex = 0;
+	long bytes = Downloaded;
+	while((bytes > 1024) && (UnitIndex++ < (UE_ARRAY_COUNT(Units) - 1)))
+	{
+		bytes >>= 10llu;
+	}
+
+	return FString::Printf(TEXT("%llu%s"), bytes, Units[UnitIndex]);
+}
\ No newline at end of file
