Moving Submit Tool to Source/Programs

diff --git a/Engine/Source/Programs/SubmitTool/Private/Logic/Validators/CrossChangelistValidator.cpp b/Engine/Source/Programs/SubmitTool/Private/Logic/Validators/CrossChangelistValidator.cpp
new file mode 100644
--- /dev/null
+++ b/Engine/Source/Programs/SubmitTool/Private/Logic/Validators/CrossChangelistValidator.cpp
@@ -0,0 +1,77 @@
+// Copyright Epic Games, Inc. All Rights Reserved.
+
+#include "CrossChangelistValidator.h"
+#include "Misc/Paths.h"
+#include "Logging/SubmitToolLog.h"
+#include "Logic/ChangelistService.h"
+
+bool FCrossChangelistValidator::Validate(const FString& InCLDescription, const TArray<FSourceControlStateRef>& InFilteredFilesInCL, const TArray<const FTag*>& InTags)
+{
+	const TArray<FSourceControlChangelistStatePtr> OtherChangelistsStates = ServiceProvider.Pin()->GetService<FChangelistService>()->GetOtherChangelistsStates();
+	bool bValid = true;
+
+	bValid &= CheckHeaderAndCppInDifferentChangelist(OtherChangelistsStates);
+
+	ValidationFinished(bValid);
+	return true;
+}
+
+bool FCrossChangelistValidator::CheckHeaderAndCppInDifferentChangelist(const TArray<FSourceControlChangelistStatePtr>& OtherChangelistsStates)
+{
+	const FString HeaderExt(TEXT(".h"));
+	const FString CppExt(TEXT(".cpp"));
+	const FString CExt(TEXT(".c"));
+
+	const TArray<FSourceControlStateRef> &FilesInChangelist = ServiceProvider.Pin()->GetService<FChangelistService>()->GetFilesInCL();
+
+	bool bValid = true;
+
+	for (FSourceControlStateRef FileInCL : FilesInChangelist)
+	{
+		const FString& Filename = FileInCL->GetFilename();
+		if ( Filename.EndsWith(HeaderExt, ESearchCase::IgnoreCase)
+		  || Filename.EndsWith(CExt, ESearchCase::IgnoreCase)
+		  || Filename.EndsWith(CppExt, ESearchCase::IgnoreCase))
+		{
+			TArray<FString> FilenamesToCheck;
+
+			if (Filename.EndsWith(HeaderExt, ESearchCase::IgnoreCase))
+			{
+				FilenamesToCheck.Add(FPaths::GetCleanFilename(Filename.Replace(*HeaderExt, *CppExt)));
+				FilenamesToCheck.Add(FPaths::GetCleanFilename(Filename.Replace(*HeaderExt, *CExt)));
+			}
+
+			if (Filename.EndsWith(CExt, ESearchCase::IgnoreCase))
+			{
+				FilenamesToCheck.Add(FPaths::GetCleanFilename(Filename.Replace(*CExt, *HeaderExt)));
+			}
+
+			if (Filename.EndsWith(CppExt, ESearchCase::IgnoreCase))
+			{
+				FilenamesToCheck.Add(FPaths::GetCleanFilename(Filename.Replace(*CppExt, *HeaderExt)));
+			}
+
+			for (const FSourceControlChangelistStatePtr& ChangelistState : OtherChangelistsStates)
+			{
+				for (const FSourceControlStateRef& FileSate : ChangelistState->GetFilesStates())
+				{
+					FString OtherFilename = FPaths::GetCleanFilename(FileSate->GetFilename());
+					if (FilenamesToCheck.Contains(OtherFilename))
+					{
+						bValid = false;
+
+						FString Message = FString::Printf(TEXT("[%s] %s file '%s' is not in the current CL, it is in CL '%s'"),
+							*GetValidatorName(),
+							OtherFilename.EndsWith(HeaderExt) ? TEXT("Header") : TEXT("CPP | C"),
+							*OtherFilename,
+							*(ChangelistState->GetChangelist()->GetIdentifier()));
+
+						LogFailure(Message);
+					}
+				}
+			}
+		}
+	}
+
+	return bValid;
+}
diff --git a/Engine/Source/Programs/SubmitTool/Private/Logic/Validators/CrossChangelistValidator.h b/Engine/Source/Programs/SubmitTool/Private/Logic/Validators/CrossChangelistValidator.h
new file mode 100644
--- /dev/null
+++ b/Engine/Source/Programs/SubmitTool/Private/Logic/Validators/CrossChangelistValidator.h
@@ -0,0 +1,17 @@
+// Copyright Epic Games, Inc. All Rights Reserved.
+
+#pragma once
+
+#include "ValidatorBase.h"
+#include "ISourceControlChangelistState.h"
+class FCrossChangelistValidator : public FValidatorBase
+{
+public:
+	using FValidatorBase::FValidatorBase;
+
+	virtual bool Validate(const FString& InCLDescription, const TArray<FSourceControlStateRef>& InFilteredFilesInCL, const TArray<const FTag*>& InTags) override;
+	virtual const FString& GetValidatorTypeName() const override { return SubmitToolParseConstants::CrossChangelistValidator; }
+
+private:
+	bool CheckHeaderAndCppInDifferentChangelist(const TArray<FSourceControlChangelistStatePtr>& OtherChangelistsStates);
+};
\ No newline at end of file
diff --git a/Engine/Source/Programs/SubmitTool/Private/Logic/Validators/ValidatorBase.cpp b/Engine/Source/Programs/SubmitTool/Private/Logic/Validators/ValidatorBase.cpp
new file mode 100644
--- /dev/null
+++ b/Engine/Source/Programs/SubmitTool/Private/Logic/Validators/ValidatorBase.cpp
@@ -0,0 +1,379 @@
+// Copyright Epic Games, Inc. All Rights Reserved.
+
+#include "ValidatorBase.h"
+#include "HAL/FileManagerGeneric.h"
+#include "Misc/Paths.h"
+#include "Misc/Timespan.h"
+#include "AnalyticsEventAttribute.h"
+#include "Configuration/Configuration.h"
+#include "Logic/TagService.h"
+#include "Logic/Services/Interfaces/ICacheDataService.h"
+#include "Logic/Services/Interfaces/ISTSourceControlService.h"
+
+FValidatorBase::FValidatorBase(const FName& InNameId, const FSubmitToolParameters& InParameters, TWeakPtr<FSubmitToolServiceProvider> InServiceProvider, const FString& InDefinition) :
+	ValidatorNameID(InNameId),
+	OptionsProvider(InNameId),
+	ServiceProvider(InServiceProvider),
+	SubmitToolParameters(InParameters),
+	Start(FDateTime::MinValue())
+{
+	ParseDefinition(InDefinition);
+	ValidatorName = Definition->CustomName.IsEmpty() ? ValidatorNameID.ToString() : Definition->CustomName;
+}
+
+void FValidatorBase::ParseDefinition(const FString& InDefinition)
+{
+	Definition = MakeUnique<FValidatorDefinition>();
+	FStringOutputDevice Errors;
+	FValidatorDefinition::StaticStruct()->ImportText(*InDefinition, const_cast<FValidatorDefinition*>(Definition.Get()), nullptr, 0, &Errors, FValidatorDefinition::StaticStruct()->GetName());
+
+	if(!Errors.IsEmpty())
+	{
+		UE_LOG(LogSubmitTool, Error, TEXT("[%s] Error loading parameter file %s"), *GetValidatorNameId().ToString(), *Errors);
+	}
+
+}
+
+FValidatorBase::~FValidatorBase()
+{
+	OnValidationFinished.Clear();
+}
+
+void FValidatorBase::StartValidation()
+{
+	RunTime = 0;
+	Start = FDateTime::UtcNow();
+	State = EValidationStates::Running;
+	ErrorListCache.Empty();
+
+	const TSharedPtr<FChangelistService>& ChangelistService = ServiceProvider.Pin()->GetService<FChangelistService>();
+	const TSharedPtr<FTagService>& TagService = ServiceProvider.Pin()->GetService<FTagService>();
+
+	FilteredFiles.RemoveAt(0, FilteredFiles.Num(), EAllowShrinking::No);
+
+	bool bIncrementalValidation = Definition->bUsesIncrementalCache && !bForceRun;
+	bForceRun = false;
+	TArray<FSourceControlStateRef> IncrementallySkippedFiles;
+	bool bAppliesToCL = AppliesToCL(ChangelistService->GetCLDescription(), ChangelistService->GetFilesInCL(), TagService->GetTagsArray(), FilteredFiles, IncrementallySkippedFiles, bIncrementalValidation);
+
+	if(bAppliesToCL && !bIsValidSetup)
+	{
+		LogFailure(FString::Printf(TEXT("[%s] Task is not correctly setup and should run in this CL"), *ValidatorName));
+		ValidationFinished(false);
+	}
+	else if(!bAppliesToCL)
+	{
+		if(IncrementallySkippedFiles.Num() != 0)
+		{
+			UE_LOG(LogValidators, Log, TEXT("[%s] All files were validated in a previous validation and are still valid. To force a validation click 'Run' in the validator list"), *ValidatorName );
+			UE_LOG(LogValidatorsResult, Log, TEXT("[%s] All files were validated in a previous validation and are still valid. To force a validation click 'Run' in the validator list"), *ValidatorName);
+		}
+		else
+		{
+			const FString Extensions = Definition->IncludeFilesWithExtension.IsEmpty() ? TEXT(".*") : *FString::Join(Definition->IncludeFilesWithExtension, TEXT("|"));
+			UE_LOG(LogValidators, Log, TEXT("[%s] No files match the filter %s{%s} %s won't run"), *ValidatorName, *Definition->IncludeFilesInDirectory, *Extensions, *ValidatorName);
+			UE_LOG(LogValidatorsResult, Log, TEXT("[%s] No files match the filter %s{%s} %s won't run"), *ValidatorName, *Definition->IncludeFilesInDirectory, *Extensions, *ValidatorName);
+		}
+		Skip();
+	}
+	else
+	{
+		if(IncrementallySkippedFiles.Num() != 0)
+		{
+			const FString FileList = FString::JoinBy(IncrementallySkippedFiles, TEXT("\n"), [](const FSourceControlStateRef& InFile) { return InFile->GetFilename(); });
+			UE_LOG(LogValidators, Log, TEXT("[%s] Skipping Files because they were already validated in a previous execution:\n%s"), *GetValidatorName(), *FileList);
+		}
+
+		if(!Validate(ChangelistService->GetCLDescription(), FilteredFiles, TagService->GetTagsArray()))
+		{
+			ValidationFinished(false);
+		}
+	}
+}
+
+
+void FValidatorBase::Tick(float InDeltaTime)
+{
+	RunTime += InDeltaTime;
+
+	if(Definition->TimeoutLimit > 0 && RunTime >= Definition->TimeoutLimit)
+	{
+		LogFailure(FString::Printf(TEXT("[%s]: %s"), *GetValidatorName(), TEXT("Timeout limit has been reached, cancelling task.")));
+
+		StopInternalValidations();
+		State = EValidationStates::Timeout;
+
+		if(OnValidationFinished.IsBound())
+		{
+			OnValidationFinished.Broadcast(*this);
+		}
+	}
+}
+
+bool FValidatorBase::Activate()
+{
+	bIsValidSetup = true;
+
+	if(Definition != nullptr)
+	{
+		if(!Definition->IncludeFilesInDirectory.IsEmpty())
+		{
+			FValidatorDefinition* ModifiableDefinition = const_cast<FValidatorDefinition*>(Definition.Get());
+			ModifiableDefinition->IncludeFilesInDirectory = FConfiguration::SubstituteAndNormalizeDirectory(ModifiableDefinition->IncludeFilesInDirectory);
+		}
+	}
+	else
+	{
+		bIsValidSetup = false;
+	}
+
+	return bIsValidSetup;
+}
+
+void FValidatorBase::InvalidateLocalFileModifications()
+{
+	if((Definition->TaskArea & ETaskArea::LocalFiles) == ETaskArea::LocalFiles && (State == EValidationStates::Valid || State == EValidationStates::Running || State == EValidationStates::Skipped))
+	{
+		FFileManagerGeneric FileManager;
+		for(const FSourceControlStateRef& File : ServiceProvider.Pin()->GetService<FChangelistService>()->GetFilesInCL())
+		{
+			bool bIncrementallySkipped;
+
+			if(AppliesToFile(File, false, bIncrementallySkipped))
+			{
+				FString Filename = File->GetFilename();
+				FFileStatData FileModifiedDate = FileManager.GetStatData(*Filename);
+				if(FileModifiedDate.ModificationTime > Start)
+				{
+					if(GetIsRunning())
+					{
+						UE_LOG(LogValidators, Warning, TEXT("File %s was modified during %s run, this task needs to be run again"), *Filename, *GetValidatorName());
+						UE_LOG(LogValidatorsResult, Warning, TEXT("File %s was modified during %s run, this task needs to be run again"), *Filename, *GetValidatorName());
+					}
+					else
+					{
+						UE_LOG(LogValidators, Warning, TEXT("File %s has been modified after %s last run, this task needs to be run again."), *Filename, *GetValidatorName());
+						UE_LOG(LogValidatorsResult, Warning, TEXT("File %s has been modified after %s last run, this task needs to be run again."), *Filename, *GetValidatorName());
+					}
+
+					Invalidate();
+					break;
+				}
+			}
+		}
+	}
+}
+
+const FString FValidatorBase::GetStatusText() const
+{
+	const FString StateStr = StaticEnum<EValidationStates>()->GetNameStringByValue(static_cast<int64>(State))
+		.Replace(TEXT("_"), TEXT(" "));
+
+	if(State == EValidationStates::Skipped || State == EValidationStates::Not_Run)
+	{
+		return StateStr;
+	}
+
+	// do not clutter the UI with uninteresting information
+	if (RunTime < 0.5f)
+	{
+		return StateStr;
+	}
+
+	return FString::Printf(TEXT("%s (%s)"),
+		*StateStr,
+		*FGenericPlatformTime::PrettyTime(RunTime));
+}
+
+const TArray<FAnalyticsEventAttribute> FValidatorBase::GetTelemetryAttributes() const
+{
+	return MakeAnalyticsEventAttributeArray(
+		TEXT("ValidatorID"), *GetValidatorNameId().ToString(),
+		TEXT("ValidatorName"), *GetValidatorName(),
+		TEXT("Status"), GetHasPassed(),
+		TEXT("Runtime"), RunTime,
+		TEXT("Stream"), ServiceProvider.Pin()->GetService<ISTSourceControlService>()->GetCurrentStreamName()
+		);
+}
+
+void FValidatorBase::ValidationFinished(const bool bHasPassed)
+{
+	if(bHasPassed)
+	{
+		UE_LOG(LogValidatorsResult, Log, TEXT("[%s]: Task Succeeded!"), *GetValidatorName());
+
+		if(Definition->bUsesIncrementalCache)
+		{
+			ServiceProvider.Pin()->GetService<ICacheDataService>()->UpdateLastValidationForFiles(ServiceProvider.Pin()->GetService<FChangelistService>()->GetCLID(), GetValidatorNameId(), GetValidationConfigId(), FilteredFiles, FDateTime::UtcNow());
+		}
+	}
+	else if(Definition->IsRequired)
+	{
+		UE_LOG(LogValidatorsResult, Error, TEXT("[%s]: Failed on Required Task!"), *GetValidatorName());
+	}
+	else
+	{
+		UE_LOG(LogValidatorsResult, Warning, TEXT("[%s]: Failed on Optional Task!"), *GetValidatorName());
+	}
+
+	if(!bHasPassed)
+	{
+		for(const FString& ErrorMsg : Definition->AdditionalValidationErrorMessages)
+		{
+			LogFailure(FString::Printf(TEXT("[%s]: %s"), *GetValidatorName(), *ErrorMsg));
+		}
+	}
+
+	State = bHasPassed ? EValidationStates::Valid : EValidationStates::Failed;
+
+	if(OnValidationFinished.IsBound())
+	{
+		OnValidationFinished.Broadcast(*this);
+	}
+}
+
+bool FValidatorBase::EvaluateTagSkip()
+{
+	TSharedPtr<FChangelistService> ChangelistService = ServiceProvider.Pin()->GetService<FChangelistService>();
+	if(Definition->SkipForbiddenTags.Num() > 0)
+	{
+		for(const FString& Tag : Definition->SkipForbiddenTags)
+		{
+			if(ChangelistService->GetCLDescription().Find(Tag, ESearchCase::IgnoreCase) != INDEX_NONE)
+			{
+				UE_LOG(LogValidators, Log, TEXT("[%s] The Description contains '%s'. %s is not allowed to be skipped"), *ValidatorName, *Tag, *ValidatorName);
+				UE_LOG(LogValidatorsResult, Log, TEXT("[%s] The Description contains '%s'. %s is not allowed to be skipped"), *ValidatorName, *Tag, *ValidatorName);
+				return false;
+			}
+		}
+	}
+
+	if(Definition->bSkipWhenAddendumInDescription && !Definition->ChangelistDescriptionAddendum.IsEmpty())
+	{
+		if(ChangelistService->GetCLDescription().Find(Definition->ChangelistDescriptionAddendum, ESearchCase::IgnoreCase) != INDEX_NONE)
+		{
+			UE_LOG(LogValidators, Log, TEXT("[%s] The Description Addendum '%s' is already present in the CL. %s won't run"), *ValidatorName, *Definition->ChangelistDescriptionAddendum, *ValidatorName);
+			UE_LOG(LogValidatorsResult, Log, TEXT("[%s] The Description Addendum '%s' is already present in the CL. %s won't run"), *ValidatorName, *Definition->ChangelistDescriptionAddendum, *ValidatorName);
+			Start = FDateTime::UtcNow();
+			State = EValidationStates::Skipped;
+			return true;
+		}
+	}
+
+	return false;
+}
+
+void FValidatorBase::SetSelectedOption(const FString& InOptionName, const FString& InOptionValue)
+{
+	UE_LOG(LogValidators, Log, TEXT("[%s] Task stopped due to a change in options, %s = %s"), *GetValidatorName(), *InOptionName, *InOptionValue);
+	CancelValidation();
+	OptionsProvider.SetSelectedOption(InOptionName, InOptionValue);
+}
+
+void FValidatorBase::PrintErrorSummary() const
+{
+	if(State == EValidationStates::Failed || State == EValidationStates::Timeout)
+	{
+		if(ErrorListCache.Num() > 0)
+		{
+			UE_LOG(LogValidators, Error, TEXT("========================[%s Errors Summary]========================"), *GetValidatorName());
+			UE_LOG(LogValidatorsResult, Error, TEXT("========================[%s Errors Summary]========================"), *GetValidatorName());
+			for(const FString& ErrorStr : ErrorListCache)
+			{
+				UE_LOG(LogValidators, Error, TEXT("%s"), *ErrorStr);
+				UE_LOG(LogValidatorsResult, Error, TEXT("%s"), *ErrorStr);
+			}
+			UE_LOG(LogValidators, Error, TEXT("================================================================"), *GetValidatorName());
+			UE_LOG(LogValidatorsResult, Error, TEXT("================================================================"), *GetValidatorName());
+		}
+	}
+}
+const FString FValidatorBase::GetValidationConfigId() const
+{
+	TStringBuilder<512> StringBuilder;
+	for(const TPair<FString, FString>& Pair : OptionsProvider.GetSelectedOptions())
+	{
+		StringBuilder.Append(Pair.Key);
+		StringBuilder.Append(TEXT("_"));
+		StringBuilder.Append(Pair.Value);
+		StringBuilder.Append(TEXT("-"));
+	}
+
+	return StringBuilder.ToString();
+}
+
+bool FValidatorBase::AppliesToFile(const FSourceControlStateRef InFile, bool InbAllowIncremental, bool& OutbIsIncrementalSkip) const
+{
+	bool bIncluded = false;
+	OutbIsIncrementalSkip = false;
+
+	if((Definition->TaskArea & ETaskArea::LocalFiles) == ETaskArea::None)
+	{
+		// For validators that do not work on local files, we always apply
+		return true;
+	}
+	
+
+	if(!InFile->IsDeleted()
+		|| (InFile->IsDeleted() && Definition->bAcceptDeletedFiles))
+	{
+		FString Filename = InFile->GetFilename();
+		FPaths::NormalizeFilename(Filename);
+
+		if(!Definition->IncludeFilesInDirectory.IsEmpty())
+		{
+			if(!Filename.StartsWith(Definition->IncludeFilesInDirectory, ESearchCase::IgnoreCase))
+			{
+				return false;
+			}
+		}
+
+		if(Definition->IncludeFilesWithExtension.IsEmpty())
+		{
+			bIncluded = true;
+		}
+
+		for(int Idx = 0; Idx < Definition->IncludeFilesWithExtension.Num(); Idx++)
+		{
+			if(Filename.EndsWith(Definition->IncludeFilesWithExtension[Idx], ESearchCase::IgnoreCase))
+			{
+				bIncluded = true;
+				break;
+			}
+		}
+
+		if(InbAllowIncremental)
+		{
+			FFileManagerGeneric FileManager;
+			FDateTime LastValidation = ServiceProvider.Pin()->GetService<ICacheDataService>()->GetLastValidationDate(ServiceProvider.Pin()->GetService<FChangelistService>()->GetCLID(), GetValidatorNameId(), GetValidationConfigId(), InFile->GetFilename());
+			FFileStatData FileModifiedDate = FileManager.GetStatData(*Filename);
+			if(LastValidation != FDateTime::MinValue() && FileModifiedDate.ModificationTime < LastValidation)
+			{
+				OutbIsIncrementalSkip = true;
+				bIncluded = false;
+			}
+		}
+
+	}
+
+	return bIncluded;
+}
+
+
+bool FValidatorBase::AppliesToCL(const FString& InCLDescription, const TArray<FSourceControlStateRef>& FilesInCL, const TArray<const FTag*>& Tags, TArray<FSourceControlStateRef>& OutFilteredFiles, TArray<FSourceControlStateRef>& OutIncrementalSkips, bool InbAllowIncremental) const
+{
+	for(const FSourceControlStateRef& File : FilesInCL)
+	{
+		bool bIsIncrementalSkip;
+		if(AppliesToFile(File, InbAllowIncremental, bIsIncrementalSkip))
+		{
+			OutFilteredFiles.Add(File);
+		}
+		else if(bIsIncrementalSkip)
+		{
+			OutIncrementalSkips.Add(File);
+		}
+	}
+
+	return OutFilteredFiles.Num() > 0;
+}
diff --git a/Engine/Source/Programs/SubmitTool/Private/Logic/Validators/ValidatorBase.h b/Engine/Source/Programs/SubmitTool/Private/Logic/Validators/ValidatorBase.h
new file mode 100644
--- /dev/null
+++ b/Engine/Source/Programs/SubmitTool/Private/Logic/Validators/ValidatorBase.h
@@ -0,0 +1,237 @@
+// Copyright Epic Games, Inc. All Rights Reserved.
+
+#pragma once
+
+
+#include "CoreMinimal.h"
+
+#include "Logging/SubmitToolLog.h"
+#include "Logic/ProcessWrapper.h"
+#include "Parameters/SubmitToolParameters.h" 
+
+#include "Logic/Services/SubmitToolServiceProvider.h"
+#include "ISourceControlState.h"
+#include "ValidatorDefinition.h"
+#include "ValidatorOptionsProvider.h"
+
+#include "ValidatorBase.generated.h"
+
+class FTag;
+struct FAnalyticsEventAttribute;
+
+namespace SubmitToolParseConstants
+{
+	const FString BranchStateValidator = TEXT("BranchStateValidator");
+	const FString ContentValidator = TEXT("ContentValidator");
+	const FString TagValidator = TEXT("TagValidator");
+	const FString UBTValidator = TEXT("UBTValidator");
+	const FString CustomValidator = TEXT("CustomValidator");
+	const FString CrossChangelistValidator = TEXT("CrossChangelistValidator");
+	const FString PreflightValidator = TEXT("PreflightValidator");
+	const FString InvalidateNode = TEXT("Invalidated");
+	const FString VirtualizationToolOp = TEXT("VirtualizationTool");
+};
+
+UENUM()
+enum class EValidationStates : uint8
+{
+	Not_Run,
+	Skipped,
+	Running,
+	Valid,
+	Failed,
+	Timeout,
+	Queued
+};
+
+template<class T>
+concept DerivedFromDefinition = std::is_base_of<FValidatorDefinition, T>::value;
+
+class FValidatorBase;
+DECLARE_MULTICAST_DELEGATE_OneParam(FOnValidatorFinished, const FValidatorBase& /*Validator*/)
+
+/**
+ * @brief Base class for validator.
+ */
+class FValidatorBase
+{
+public:
+	FValidatorBase() = delete;
+	FValidatorBase(const FName& InNameId, const FSubmitToolParameters& InParameters, const TWeakPtr<FSubmitToolServiceProvider> InServiceProvider, const FString& InDefinition);
+	virtual ~FValidatorBase();
+
+	virtual void ParseDefinition(const FString& InDefinition);
+
+	virtual const FString& GetValidatorTypeName() const = 0;
+	const FName& GetValidatorNameId() const
+	{
+		return ValidatorNameID;
+	}
+	const FString& GetValidatorName() const
+	{
+		return ValidatorName;
+	}
+
+	void StartValidation();
+
+	void Invalidate(bool bForce = false)
+	{
+		if(State == EValidationStates::Running)
+		{
+			CancelValidation();
+		}
+		else if(State != EValidationStates::Queued || bForce)
+		{
+			State = EValidationStates::Not_Run;
+		}
+	}
+
+	bool GetHasPassed() const
+	{
+		return State == EValidationStates::Valid || State == EValidationStates::Skipped;
+	}
+
+	bool GetIsRunningOrQueued() const
+	{
+		return GetIsRunning() || GetIsQueued();
+	}
+
+	bool GetIsRunning() const
+	{
+		return State == EValidationStates::Running;
+	}
+
+	bool GetIsQueued() const
+	{
+		return State == EValidationStates::Queued;
+	}
+
+	float GetRunTime() const
+	{
+		return RunTime;
+	}
+
+	const EValidationStates& GetState()
+	{
+		return State;
+	}
+
+	virtual void Tick(float InDeltaTime);
+
+	virtual void CancelValidation()
+	{
+		if(State == EValidationStates::Running)
+		{
+			StopInternalValidations();
+		}
+
+		State = EValidationStates::Not_Run;
+	}
+
+	virtual void SetQueued(bool InbForceRun = false)
+	{
+		bForceRun = InbForceRun;
+		if(State == EValidationStates::Running)
+		{
+			StopInternalValidations();
+		}
+
+		State = EValidationStates::Queued;
+	}
+
+	virtual bool Activate();
+
+	virtual void InvalidateLocalFileModifications();
+	virtual bool EvaluateTagSkip();
+
+	const FString GetStatusText() const;
+	const EValidationStates GetValidatorState() const
+	{
+		return State;
+	}
+
+	const TMap<FString, TMap<FString, FString>>& GetValidatorOptions() const
+	{
+		return OptionsProvider.GetValidatorOptions();
+	}
+
+	const FString GetSelectedOptionValue(const FString& InOptionName) const
+	{
+		return OptionsProvider.GetSelectedOptionValue(InOptionName);
+	}
+	const FString GetSelectedOptionKey(const FString& InOptionName) const
+	{
+		return OptionsProvider.GetSelectedOptionKey(InOptionName);
+	}
+	EValidatorOptionType GetOptionType(const FString& InOptionName) const
+	{
+		return OptionsProvider.GetOptionType(InOptionName);
+	}
+
+	void SetSelectedOption(const FString& InOptionName, const FString& InOptionValue);
+
+	virtual const TArray<FAnalyticsEventAttribute> GetTelemetryAttributes() const;
+	virtual void PrintErrorSummary() const;
+
+	const TArray<FString>& GetErrorList() const
+	{
+		return ErrorListCache;
+	}
+
+	const FString GetValidationConfigId() const;
+
+	template<DerivedFromDefinition T>
+	const T* GetTypedDefinition() const;
+
+	TArray<FName> Dependants;
+	TUniquePtr<const FValidatorDefinition> Definition;
+	FOnValidatorFinished OnValidationFinished;
+
+protected:
+	FName ValidatorNameID;
+	FString ValidatorName;
+	FValidatorOptionsProvider OptionsProvider;
+	TWeakPtr<FSubmitToolServiceProvider> ServiceProvider;
+	TArray<FSourceControlStateRef> FilteredFiles;
+	mutable TArray<FString> ErrorListCache;
+	bool bIsValidSetup = false;
+	bool bForceRun = false;
+	const FSubmitToolParameters& SubmitToolParameters;
+
+	FDateTime Start;
+	float RunTime = 0;
+
+	/**
+	 * @brief Starts the validator.
+	 * @param SubmitSettings The submit settings to use for this validation cycle.
+	 * @return Returns true if validation has started successfully, false if validation failed.
+	 */
+	virtual bool Validate(const FString& CLDescription, const TArray<FSourceControlStateRef>& FilteredFilesInCL, const TArray<const FTag*>& Tags) = 0;
+
+	virtual bool AppliesToFile(const FSourceControlStateRef InFile, bool InbAllowIncremental, bool& OutbIsIncrementalSkip) const;
+	virtual bool AppliesToCL(const FString& CLDescription, const TArray<FSourceControlStateRef>& FilesInCL, const TArray<const FTag*>& Tags, TArray<FSourceControlStateRef>& OutFilteredFiles, TArray<FSourceControlStateRef>& OutIncrementalSkips, bool InbAllowIncremental) const;
+	virtual void StopInternalValidations()
+	{};
+	virtual void ValidationFinished(const bool bHasPassed);
+
+	virtual void LogFailure(const FString& FormattedMsg) const
+	{
+		UE_LOG(LogValidators, Error, TEXT("%s"), *FormattedMsg);
+		UE_LOG(LogValidatorsResult, Error, TEXT("%s"), *FormattedMsg);
+
+		ErrorListCache.Add(FormattedMsg);
+	}
+
+	virtual void Skip()
+	{
+		State = EValidationStates::Skipped; OnValidationFinished.Broadcast(*this);
+	}
+private:
+	EValidationStates State = EValidationStates::Not_Run;
+};
+
+template<DerivedFromDefinition T>
+inline const T* FValidatorBase::GetTypedDefinition() const
+{
+	return  static_cast<const T*>(Definition.Get());
+}
diff --git a/Engine/Source/Programs/SubmitTool/Private/Logic/Validators/ValidatorRunExecutable.cpp b/Engine/Source/Programs/SubmitTool/Private/Logic/Validators/ValidatorRunExecutable.cpp
new file mode 100644
--- /dev/null
+++ b/Engine/Source/Programs/SubmitTool/Private/Logic/Validators/ValidatorRunExecutable.cpp
@@ -0,0 +1,474 @@
+// Copyright Epic Games, Inc. All Rights Reserved.
+
+
+#include "ValidatorRunExecutable.h"
+
+#include "AnalyticsEventAttribute.h"
+#include "CommandLine/CmdLineParameters.h"
+#include "Configuration/Configuration.h"
+#include "HAL/FileManagerGeneric.h"
+#include "Internationalization/Regex.h"
+#include "Misc/FileHelper.h"
+#include "Misc/PathViews.h"
+#include "Misc/Paths.h"
+#include "Models/ModelInterface.h"
+#include "Models/SubmitToolUserPrefs.h"
+#include "Modules/BuildVersion.h"
+#include "Policies/CondensedJsonPrintPolicy.h"
+#include "Serialization/JsonReader.h"
+#include "Serialization/JsonSerializer.h"
+#include "Serialization/JsonWriter.h"
+
+FValidatorRunExecutable::FValidatorRunExecutable(const FName& InNameId, const FSubmitToolParameters& InParameters, TSharedRef<FSubmitToolServiceProvider> InServiceProvider, const FString& InDefinition) :
+	FValidatorBase(InNameId, InParameters, InServiceProvider, InDefinition)
+{
+	ParseDefinition(InDefinition);
+}
+
+void FValidatorRunExecutable::ParseDefinition(const FString& InDefinition)
+{
+	FStringOutputDevice Errors;
+	Definition = MakeUnique<FValidatorRunExecutableDefinition>();
+	FValidatorRunExecutableDefinition* ModifyableDefinition = const_cast<FValidatorRunExecutableDefinition*>(GetTypedDefinition<FValidatorRunExecutableDefinition>());
+	FValidatorRunExecutableDefinition::StaticStruct()->ImportText(*InDefinition, ModifyableDefinition, nullptr, 0, &Errors, FValidatorRunExecutableDefinition::StaticStruct()->GetName());
+
+	if(!Errors.IsEmpty())
+	{
+		UE_LOG(LogSubmitTool, Error, TEXT("Error loading parameter file %s"), *Errors);
+		FModelInterface::SetErrorState();
+	}
+}
+
+bool FValidatorRunExecutable::Validate(const FString& InCLDescription, const TArray<FSourceControlStateRef>& InFilteredFilesInCL, const TArray<const FTag*>& InTags)
+{
+	const FValidatorRunExecutableDefinition* TypedDefinition = GetTypedDefinition<FValidatorRunExecutableDefinition>();
+
+	FString FinalArgs{ FConfiguration::Substitute(TypedDefinition->ExecutableArguments) };
+
+	TArray<FString> Files;
+
+	// aggregate all the valid files
+	for(const FSourceControlStateRef& File : InFilteredFilesInCL)
+	{
+		Files.Add(File.Get().GetFilename());
+	}
+
+	if(!bIsValidSetup)
+	{
+		LogFailure(FString::Printf(TEXT("[%s] This task is not correctly setup and it's required for this change"), *ValidatorName));
+		return false;
+	}
+
+	// create a file with the files and pass it as an argument of the validator (better to avoid breaking command line length limits)
+	if(!TypedDefinition->FileListArgument.IsEmpty())
+	{
+		FString ValidatorDirectory = FPaths::EngineDir() + TEXT("Intermediate/SubmitTool/FileLists/");
+
+		FGuid guid = FGuid::NewGuid();
+		FString FileListPath = FPaths::ConvertRelativePathToFull(ValidatorDirectory + guid.ToString(EGuidFormats::DigitsWithHyphens) + TEXT(".txt"));
+
+		FFileHelper::SaveStringArrayToFile(Files, *FileListPath, FFileHelper::EEncodingOptions::ForceAnsi, &IFileManager::Get(), EFileWrite::FILEWRITE_None);
+
+		FinalArgs += FString::Printf(TEXT(" %s\"%s\""), *TypedDefinition->FileListArgument, *FileListPath);
+	}
+
+	// Legacy, pass down each file to the validator in the command line
+	else if(!TypedDefinition->FileInCLArgument.IsEmpty())
+	{
+		for(const FString& File : Files)
+		{
+			FinalArgs += " " + TypedDefinition->FileInCLArgument + File;
+		}
+	}
+
+	FString ExecutablePath;
+	if(TypedDefinition->ExecutableCandidates.Num() != 0)
+	{
+		ExecutablePath = OptionsProvider.GetSelectedOptionValue(ExecutableOptions);
+	}
+	else
+	{
+		ExecutablePath = TypedDefinition->ExecutablePath;
+	}
+
+	return StartProcess(ExecutablePath, FinalArgs);
+}
+
+bool FValidatorRunExecutable::Activate()
+{
+	bIsValidSetup = FValidatorBase::Activate();
+
+	PrepareExecutableOptions();
+	const FValidatorRunExecutableDefinition* TypedDefinition = GetTypedDefinition<FValidatorRunExecutableDefinition>();
+
+	FString SelectedOption = OptionsProvider.GetSelectedOptionKey(ExecutableOptions);
+	if(TypedDefinition->ExecutableCandidates.Num() > 0 && SelectedOption.IsEmpty())
+	{
+		TArray<FString> ExecutablePaths;
+		TypedDefinition->ExecutableCandidates.GenerateValueArray(ExecutablePaths);
+		UE_LOG(LogSubmitTool, Error, TEXT("[%s] None of the executable candidates exists locally:\n%s"), *ValidatorName, *FString::Join(ExecutablePaths, TEXT("\n")));
+		bIsValidSetup = false;
+	}
+
+	FValidatorRunExecutableDefinition* ModifyableDefinition = const_cast<FValidatorRunExecutableDefinition*>(TypedDefinition);
+
+	ModifyableDefinition->ExecutablePath = FConfiguration::SubstituteAndNormalizeFilename(TypedDefinition->ExecutablePath);
+
+	static TArray<FString> ValidExtensions =
+	{
+#if PLATFORM_WINDOWS
+		TEXT(".exe"),
+		TEXT(".bat"),
+#elif PLATFORM_MAC
+		TEXT(".app"),
+		TEXT(".sh"),
+		TEXT(".command"),
+		TEXT(""),
+#elif PLATFORM_LINUX
+		TEXT(".sh"),
+		TEXT(""),
+#endif
+	};
+
+	if(TypedDefinition->ExecutablePath.IsEmpty())
+	{
+		for(const TTuple<FString,FString>& Paths : TypedDefinition->ExecutableCandidates)
+		{
+			if(!ValidExtensions.Contains(FPaths::GetExtension(Paths.Value, true)))
+			{
+				UE_LOG(LogSubmitTool, Error, TEXT("Task '%s' executable has an invalid extension for this platform: %s"), *ValidatorName, *Paths.Value);
+				bIsValidSetup = false;
+			}
+		}
+
+		if(TypedDefinition->ExecutableCandidates.Num() == 0)
+		{
+			UE_LOG(LogSubmitTool, Error, TEXT("Task '%s' does not have a value for 'ExecutablePath' or 'ExecutableCandidates'."), *ValidatorName);
+			bIsValidSetup = false;
+		}
+	}
+	else
+	{
+		if(!ValidExtensions.Contains(FPaths::GetExtension(TypedDefinition->ExecutablePath, true)))
+		{
+			UE_LOG(LogSubmitTool, Error, TEXT("Task '%s' executable has an invalid extension for this platform: %s"), *ValidatorName, *TypedDefinition->ExecutablePath);
+			bIsValidSetup = false;
+		}
+
+		if(!FPaths::FileExists(TypedDefinition->ExecutablePath))
+		{
+			UE_LOG(LogSubmitTool, Error, TEXT("Task '%s' executable is not found on disk: %s."), *ValidatorName, *TypedDefinition->ExecutablePath);
+			bIsValidSetup = false;
+		}
+	}
+
+	if(!TypedDefinition->ExecutablePath.IsEmpty() && TypedDefinition->ExecutableCandidates.Num() > 0)
+	{
+		UE_LOG(LogSubmitTool, Error, TEXT("Specifying ExecutablePath and ExecutableCandidates for task %s is not supported, please check your config."), *GetValidatorName());
+	}
+
+
+	return bIsValidSetup;
+}
+
+void FValidatorRunExecutable::StopInternalValidations()
+{
+	if(GetValidatorState() == EValidationStates::Running)
+	{
+		if(ProcessWrapper.IsValid())
+		{
+			ProcessWrapper->Stop();
+		}
+	}
+}
+
+bool FValidatorRunExecutable::StartProcess(const FString& LocalPath, const FString& Args)
+{
+	// run validators from the $(root) dir
+	FString WorkingDir = FConfiguration::Substitute(TEXT("$(root)"));
+	ProcessWrapper = MakeUnique<FProcessWrapper>(GetValidatorName(), LocalPath, Args, FOnCompleted::CreateRaw(this, &FValidatorRunExecutable::OnProcessComplete), FOnOutputLine::CreateRaw(this, &FValidatorRunExecutable::OnProcessOutputLine), WorkingDir, GetTypedDefinition<FValidatorRunExecutableDefinition>()->bLaunchHidden, GetTypedDefinition<FValidatorRunExecutableDefinition>()->bLaunchReallyHidden);
+	const bool ProcessStarted = ProcessWrapper->Start();
+	bIgnoringOutputErrors = !GetTypedDefinition<FValidatorRunExecutableDefinition>()->EnableOutputErrorsAnchor.IsEmpty();
+
+	if (ProcessStarted)
+	{
+		OnProcessOutputLine(FString(TEXT("Task process started.")), EProcessOutputType::ProcessInfo);
+	}
+	else
+	{
+		FString ErrorMessage = FString::Format(TEXT("Task process failed to start with Process path: '{0}' and arguments: '{1}'"),
+			{
+				*LocalPath,
+				*Args
+			});
+
+		OnProcessOutputLine(ErrorMessage, EProcessOutputType::ProcessError);
+		ValidationFinished(false);
+	}
+
+	return ProcessStarted;
+}
+
+void FValidatorRunExecutable::OnProcessOutputLine(const FString& Line, const EProcessOutputType& OutputType)
+{
+	if(OutputType == EProcessOutputType::ProcessError || IsLineAnError(Line))
+	{
+		LogFailure(FString::Printf(TEXT("[%s]: %s"), *GetValidatorName(), *Line));
+	}
+	else
+	{
+		UE_LOG(LogValidators, Log, TEXT("[%s]: %s"), *GetValidatorName(), *Line);
+	}
+}
+
+bool FValidatorRunExecutable::IsLineAnError(const FString& InLine)
+{
+	const FValidatorRunExecutableDefinition* TypedDefinition = GetTypedDefinition<FValidatorRunExecutableDefinition>();
+
+	if(!TypedDefinition->EnableOutputErrorsAnchor.IsEmpty() && InLine.Find(TypedDefinition->EnableOutputErrorsAnchor, ESearchCase::IgnoreCase) != INDEX_NONE)
+	{
+		bIgnoringOutputErrors = false;
+		return false;
+	}
+	
+	if(!TypedDefinition->DisableOutputErrorsAnchor.IsEmpty() && InLine.Find(TypedDefinition->DisableOutputErrorsAnchor, ESearchCase::IgnoreCase) != INDEX_NONE)
+	{
+		bIgnoringOutputErrors = true;
+		return false;
+	}
+
+	if(bIgnoringOutputErrors)
+	{
+		return false;
+	}
+
+	for(const FString& Message : TypedDefinition->IgnoredErrorMessages)
+	{
+		if(InLine.Contains(Message, ESearchCase::IgnoreCase))
+		{
+			return false;
+		}
+	}
+
+	for(const FString& Message : TypedDefinition->ErrorMessages)
+	{
+		if(InLine.Contains(Message, ESearchCase::IgnoreCase))
+		{
+			return true;
+		}
+	}
+
+	{
+		FRegexPattern Pattern = FRegexPattern(TypedDefinition->RegexErrorParsing, ERegexPatternFlags::CaseInsensitive);
+		FRegexMatcher Regex = FRegexMatcher(Pattern, InLine);
+		if(Regex.FindNext())
+		{
+			return true;
+		}
+	}
+
+	if(Definition->bTreatWarningsAsErrors)
+	{
+		FRegexPattern Pattern = FRegexPattern(TypedDefinition->RegexWarningParsing, ERegexPatternFlags::CaseInsensitive);
+		FRegexMatcher Regex = FRegexMatcher(Pattern, InLine);
+		if(Regex.FindNext())
+		{
+			return true;
+		}
+	}
+
+	return false;
+}
+
+void FValidatorRunExecutable::OnProcessComplete(int32 ReturnCode)
+{
+	const FValidatorRunExecutableDefinition* Def = GetTypedDefinition<FValidatorRunExecutableDefinition>();
+	const bool Success = Def->AllowedExitCodes.Contains(ReturnCode) && (Def->bOnlyLookAtExitCode || ErrorListCache.IsEmpty());
+
+	if(Success)
+	{
+		UE_LOG(LogValidators, Log, TEXT("[%s]: Task process succeded (Exit code %d)"), *GetValidatorName(), ReturnCode);
+	}
+	else
+	{
+		LogFailure(FString::Printf(TEXT("[%s]: Task process failed with exit code %d and %d log errors."), *GetValidatorName(), ReturnCode, ErrorListCache.Num()));
+	}
+
+	ValidationFinished(Success);
+}
+
+const TArray<FAnalyticsEventAttribute> FValidatorRunExecutable::GetTelemetryAttributes() const
+{
+	TArray<FAnalyticsEventAttribute> Attributes = FValidatorBase::GetTelemetryAttributes();
+
+	if (ProcessWrapper)
+	{
+		return AppendAnalyticsEventAttributeArray(Attributes,
+			TEXT("ExeExitCode"), ProcessWrapper->ExitCode,
+			TEXT("ExeRunTime"), ProcessWrapper->ExecutingTime,
+			TEXT("ErrorCount"), ErrorListCache.Num()
+			);
+	}
+	else
+	{
+		return Attributes;
+	}	
+}
+
+bool FValidatorRunExecutable::DoesExecutableNeedBuilding() const
+{
+	const FValidatorRunExecutableDefinition* ExeDefinition = GetTypedDefinition<FValidatorRunExecutableDefinition>();
+	check(ExeDefinition != nullptr);
+
+	bool bExeNeedsBuilding = false;
+
+	if (IFileManager::Get().FileExists(*ExeDefinition->ExecutablePath))
+	{
+		FBuildVersion VersionInfo;
+		if (FindBuildVersionForExecutable(ExeDefinition->ExecutablePath, VersionInfo))
+		{
+			if (!VersionInfo.BuildUrl.IsEmpty())
+			{
+				UE_LOG(LogValidators, Log, TEXT("[%s] BuildVersion info for '%s' indicates that it is a precompiled binary"), *ValidatorName, *ExeDefinition->ExecutablePath);
+				return false;
+			}
+			else
+			{
+				UE_LOG(LogValidators, Log, TEXT("[%s] BuildVersion info for '%s' indicates that it was built locally"), *ValidatorName, *ExeDefinition->ExecutablePath);
+				return true;
+			}
+		}
+		else
+		{
+			UE_LOG(LogValidators, Warning, TEXT("[%s] Failed to retrieve BuildVersion info for '%s', assuming that it was locally built"), *ValidatorName, *ExeDefinition->ExecutablePath);
+			return true;
+		}
+	}
+	else
+	{
+		UE_LOG(LogValidators, Log, TEXT("[%s] Failed to find '%s', so it will need to be built locally"), *ValidatorName, *ExeDefinition->ExecutablePath);
+		return true;
+	}
+}
+
+bool FValidatorRunExecutable::FindBuildVersionForExecutable(const FString& ExecutablePath, FBuildVersion& OutBuildVersion) const
+{
+	const FString VersionPath = FPathViews::ChangeExtension(ExecutablePath, TEXTVIEW("version"));
+	if (IFileManager::Get().FileExists(*VersionPath))
+	{
+		FBuildVersion BuildVersion;
+		if (FBuildVersion::TryRead(VersionPath, BuildVersion))
+		{
+			OutBuildVersion = MoveTemp(BuildVersion);
+			return true;
+		}
+		else
+		{
+			return false;
+		}
+	}
+
+	const FString TargetPath = FPathViews::ChangeExtension(ExecutablePath, TEXTVIEW("target"));
+	if (IFileManager::Get().FileExists(*TargetPath))
+	{
+		FString JsonString;
+		if (!FFileHelper::LoadFileToString(JsonString, *TargetPath))
+		{
+			return false;
+		}
+
+		TSharedRef<TJsonReader<>> JsonReader = TJsonReaderFactory<>::Create(JsonString);
+		TSharedPtr<FJsonObject> JsonRootObject;
+		if (!FJsonSerializer::Deserialize(JsonReader, JsonRootObject))
+		{
+			return false;
+		}
+
+		const TSharedPtr<FJsonObject>* JsonVersionObject;
+		if (!JsonRootObject->TryGetObjectField(TEXTVIEW("version"), JsonVersionObject))
+		{
+			return false;
+		}
+
+		FString JsonObjectString;
+		TSharedRef<TJsonWriter<TCHAR, TCondensedJsonPrintPolicy<TCHAR>>> JsonWriter = TJsonWriterFactory<TCHAR, TCondensedJsonPrintPolicy<TCHAR>>::Create(&JsonObjectString, 0);
+		if (!FJsonSerializer::Serialize((*JsonVersionObject).ToSharedRef(), JsonWriter, true))
+		{
+			return false;
+		}
+
+		FBuildVersion BuildVersion;
+		if (FBuildVersion::TryReadFromString(JsonObjectString, BuildVersion))
+		{
+			OutBuildVersion = MoveTemp(BuildVersion);
+			return true;
+		}
+		else
+		{
+			return false;
+		}
+
+	}
+
+	return false;
+}
+
+void FValidatorRunExecutable::PrepareExecutableOptions()
+{
+	const FValidatorRunExecutableDefinition* TypedDefinition = GetTypedDefinition<FValidatorRunExecutableDefinition>();
+	if(TypedDefinition->ExecutableCandidates.Num() > 0)
+	{
+		TMap<FString, FString> Options;
+		Options.Reserve(TypedDefinition->ExecutableCandidates.Num());
+
+		FValidatorRunExecutableDefinition* ModifyableDefinition = const_cast<FValidatorRunExecutableDefinition*>(TypedDefinition);
+
+		FString SelectedOption;
+		FString* UserSelectedOption = FSubmitToolUserPrefs::Get()->ValidatorOptions.Find(OptionsProvider.GetUserPrefsKey(ExecutableOptions));
+		if(UserSelectedOption != nullptr && !UserSelectedOption->StartsWith(TEXT("Auto Select"), ESearchCase::IgnoreCase))
+		{
+			SelectedOption = *UserSelectedOption;
+		}	
+
+		FDateTime LastWriteAccess = FDateTime::MinValue();
+		FFileManagerGeneric FileManager;
+
+		FString NewestExecutable;	
+
+		for(TPair<FString, FString>& ExecutableCandidate : ModifyableDefinition->ExecutableCandidates)
+		{
+			ExecutableCandidate.Value = FConfiguration::SubstituteAndNormalizeFilename(ExecutableCandidate.Value);
+			Options.Add(ExecutableCandidate.Key, ExecutableCandidate.Value);
+
+			if(FPaths::FileExists(*ExecutableCandidate.Value))
+			{
+				if(TypedDefinition->bUseLatestExecutable)
+				{
+					FFileStatData FileModifiedDate = FileManager.GetStatData(*ExecutableCandidate.Value);
+					if(FileModifiedDate.ModificationTime > LastWriteAccess)
+					{
+						NewestExecutable = ExecutableCandidate.Key;
+						LastWriteAccess = FileModifiedDate.ModificationTime;
+					}
+				}
+				else if(SelectedOption.IsEmpty())
+				{
+					SelectedOption = ExecutableCandidate.Key;
+				}
+			}
+		}
+
+		if(TypedDefinition->bUseLatestExecutable && TypedDefinition->ExecutableCandidates.Contains(NewestExecutable))
+		{
+			const FString NewestExecutableOptionKey = FString::Printf(TEXT("Auto Select (%s)"), *NewestExecutable);
+			Options.Add(NewestExecutableOptionKey, TypedDefinition->ExecutableCandidates[NewestExecutable]);
+
+			if(UserSelectedOption == nullptr || UserSelectedOption->StartsWith(TEXT("Auto Select"), ESearchCase::IgnoreCase))
+			{
+				SelectedOption = NewestExecutableOptionKey;
+			}
+		}
+
+		OptionsProvider.InitializeValidatorOptions(ExecutableOptions, Options, SelectedOption, EValidatorOptionType::FilePath);
+	}
+}
diff --git a/Engine/Source/Programs/SubmitTool/Private/Logic/Validators/ValidatorRunExecutable.h b/Engine/Source/Programs/SubmitTool/Private/Logic/Validators/ValidatorRunExecutable.h
new file mode 100644
--- /dev/null
+++ b/Engine/Source/Programs/SubmitTool/Private/Logic/Validators/ValidatorRunExecutable.h
@@ -0,0 +1,42 @@
+// Copyright Epic Games, Inc. All Rights Reserved.
+
+#pragma once
+
+#include "ValidatorBase.h"
+
+class FBuildVersion;
+
+/**
+ * @brief Runs an executable with optional arguments. Return code 0 on the executable means success.
+ * Validator values: "ExecutablePath" (required), "ExecutableArguments" (optional)
+ */
+class FValidatorRunExecutable : public FValidatorBase
+{
+public:	
+	FValidatorRunExecutable(const FName& InNameId, const FSubmitToolParameters& InParameters, TSharedRef<FSubmitToolServiceProvider> InServiceProvider, const FString& InDefinition);
+	virtual void ParseDefinition(const FString& InDefinition) override;
+
+	virtual bool Validate(const FString& InCLDescription, const TArray<FSourceControlStateRef>& InFilteredFilesInCL, const TArray<const FTag*>& InTags) override;
+	virtual const FString& GetValidatorTypeName() const override { return SubmitToolParseConstants::CustomValidator; }
+	virtual bool Activate() override;
+
+	virtual void StopInternalValidations() override;
+	 
+protected:
+	virtual bool StartProcess(const FString& LocalPath, const FString& Args);
+	virtual void OnProcessOutputLine(const FString& Line, const EProcessOutputType& OutputType);
+	virtual bool IsLineAnError(const FString& InLine);
+	virtual void OnProcessComplete(int32 ReturnCode);
+
+	virtual const TArray<FAnalyticsEventAttribute> GetTelemetryAttributes() const override;
+
+	bool DoesExecutableNeedBuilding() const;
+	bool FindBuildVersionForExecutable(const FString& ExecutablePath, FBuildVersion& OutBuildVersion) const;
+
+	void PrepareExecutableOptions();
+
+protected:
+	const FString ExecutableOptions = TEXT("SelectedExecutable");
+	TUniquePtr<FProcessWrapper> ProcessWrapper;
+	bool bIgnoringOutputErrors = false;
+};
diff --git a/Engine/Source/Programs/SubmitTool/Private/Models/ModelInterface.cpp b/Engine/Source/Programs/SubmitTool/Private/Models/ModelInterface.cpp
new file mode 100644
--- /dev/null
+++ b/Engine/Source/Programs/SubmitTool/Private/Models/ModelInterface.cpp
@@ -0,0 +1,907 @@
+// Copyright Epic Games, Inc. All Rights Reserved.
+
+#include "ModelInterface.h"
+#include "SubmitToolUserPrefs.h"
+#include "Logic/JiraService.h"
+#include "Logic/PreflightService.h"
+#include "Framework/Application/SlateApplication.h"
+#include "Version/AppVersion.h"
+#include "Widgets/Docking/SDockTab.h"
+#include "Containers/Ticker.h"
+#include "Configuration/Configuration.h"
+
+#include "Logic/Validators/PreflightValidator.h"
+#include "Logic/Validators/TagValidator.h"
+#include "Logic/Validators/UBTValidator.h"
+#include "Logic/Validators/ValidatorBase.h"
+#include "Logic/Validators/ValidatorRunExecutable.h"
+#include "Logic/Validators/CrossChangelistValidator.h"
+#include "Logic/PreSubmitOperations/VirtualizationOperation.h"
+#include "Logic/Services/CacheDataService.h"
+#include "Logic/Services/SourceControl/SubmitToolPerforce.h"
+
+ESubmitToolAppState FModelInterface::SubmitToolState = ESubmitToolAppState::Initializing;
+
+FModelInterface::FModelInterface(const FSubmitToolParameters& InParameters) :
+	Parameters(InParameters)
+{
+	// initialize call backs
+	CLReadyCallback = FOnChangeListReadyDelegate::CreateRaw(this, &FModelInterface::OnChangelistReady);
+	CLRefreshCallback = FOnChangelistRefreshDelegate::CreateRaw(this, &FModelInterface::OnChangelistRefresh);
+	SubmitFinishedCallback = FSourceControlOperationComplete::CreateRaw(this, &FModelInterface::OnSubmitOperationComplete);
+	DeleteShelveCallback = FSourceControlOperationComplete::CreateRaw(this, &FModelInterface::OnDeleteShelveOperationComplete);
+	RevertUnchangedCallback = FSourceControlOperationComplete::CreateRaw(this, &FModelInterface::OnRevertUnchangedOperationComplete);
+
+	ServiceProvider = MakeShared<FSubmitToolServiceProvider>();
+	// Initialize services
+	if(Parameters.GeneralParameters.CacheFile.IsEmpty())
+	{
+		ServiceProvider->RegisterService<ICacheDataService>(MakeShared<FNoOpCacheDataService>());
+	}
+	else
+	{
+		ServiceProvider->RegisterService<ICacheDataService>(MakeShared<FCacheDataService>(Parameters.GeneralParameters));
+	}
+
+	SourceControlService = MakeShared<FSubmitToolPerforce>(InParameters);
+	ServiceProvider->RegisterService<ISTSourceControlService>(SourceControlService.ToSharedRef());
+	ValidationService = MakeShared<FTasksService>(InParameters.Validators, TEXT("SubmitTool.StandAlone.Validator"));
+	ServiceProvider->RegisterService<FTasksService>(ValidationService.ToSharedRef());
+	PresubmitOperationsService = MakeShared<FTasksService>(InParameters.PresubmitOperations, TEXT("SubmitTool.StandAlone.PresubmitOperation"));
+	ServiceProvider->RegisterService<FTasksService>(PresubmitOperationsService.ToSharedRef());
+	CredentialsService = MakeShared<FCredentialsService>(InParameters.OAuthParameters);
+	ServiceProvider->RegisterService<FCredentialsService>(CredentialsService.ToSharedRef());
+	ChangelistService = MakeShared<FChangelistService>(InParameters.GeneralParameters, SourceControlService, CLReadyCallback, CLRefreshCallback);
+	ServiceProvider->RegisterService<FChangelistService>(ChangelistService.ToSharedRef());
+	P4LockdownService = MakeShared<FP4LockdownService>(InParameters.P4LockdownParameters, ServiceProvider.ToSharedRef());
+	ServiceProvider->RegisterService<FP4LockdownService>(P4LockdownService.ToSharedRef());
+	TagService = MakeShared<FTagService>(InParameters, ChangelistService);
+	ServiceProvider->RegisterService<FTagService>(TagService.ToSharedRef());
+	SwarmService = MakeShared<FSwarmService>(ServiceProvider.ToSharedRef());
+	ServiceProvider->RegisterService<FSwarmService>(SwarmService.ToSharedRef());
+	PreflightService = MakeShared<FPreflightService>(InParameters.HordeParameters, this, ServiceProvider.ToSharedRef());
+	ServiceProvider->RegisterService<FPreflightService>(PreflightService.ToSharedRef());
+	JiraService = MakeShared<FJiraService>(InParameters.JiraParameters, 256, ServiceProvider.ToSharedRef());
+	ServiceProvider->RegisterService<FJiraService>(JiraService.ToSharedRef());
+	FNIntegrationService = MakeShared<FIntegrationService>(InParameters.IntegrationParameters, ServiceProvider);
+	ServiceProvider->RegisterService<FIntegrationService>(FNIntegrationService.ToSharedRef());
+	UpdateService = MakeShared<FUpdateService>(InParameters.HordeParameters, InParameters.AutoUpdateParameters, ServiceProvider);
+	ServiceProvider->RegisterService<FUpdateService>(UpdateService.ToSharedRef());
+
+	ParseValidators();
+	ParsePreSubmitOperations();
+
+	OnValidationStateUpdatedHandle = ValidationService->OnTasksRunResultUpdated.Add(FOnTaskRunStateChanged::FDelegate::CreateLambda([this](bool bIsValid) {
+		if(bIsValid)
+		{
+			bool bOptionalFailures = false;
+			for(const TWeakPtr<const FValidatorBase>& Validator : ValidationService->GetTasks())
+			{
+				if(!Validator.Pin()->GetIsRunningOrQueued() && !Validator.Pin()->GetHasPassed())
+				{
+					bOptionalFailures = true;
+				}
+			}
+
+			UE_LOG(LogSubmitTool, Log, TEXT("The required local validation has succeeded, you're ALLOWED TO SUBMIT."))
+			if(ValidationService->GetIsAnyTaskRunning())
+			{
+				UE_LOG(LogSubmitTool, Warning, TEXT("You still have optional validations running you might want to consider waiting for them to finish."))
+			}
+			if(bOptionalFailures)
+			{
+				UE_LOG(LogSubmitTool, Warning, TEXT("You have optional validations that have failed, you can still proceed with the submission if you consider that these failures are not relevant. Please make sure this is the case."))
+			}
+		}
+	}));
+	
+	OnSingleValidationFinishedHandle = ValidationService->OnSingleTaskFinished.AddLambda([this](const FValidatorBase& InTask) {
+			ReevaluateSubmitToolTag();
+
+			if(bPreflightQueued && CanLaunchPreflight())
+			{
+				bPreflightQueued = false;
+				PreflightService->RequestPreflight();
+			}
+		});
+
+	OnValidationFinishedHandle = ValidationService->OnTasksQueueFinished.Add(FOnTaskFinished::FDelegate::CreateLambda([this](bool bIsValid)
+		{
+			MainTab.Pin()->GetParentWindow()->DrawAttention(FWindowDrawAttentionParameters());
+			
+			if(bIsValid)
+			{
+				if (!bPreflightQueued && bSubmitOnSuccessfulValidation && !IsIntegrationRequired())
+				{
+					bool bAllSucceedIncludingOptional = true;
+					for(const TWeakPtr<const FValidatorBase>& Validator : ValidationService->GetTasks())
+					{
+						if(Validator.IsValid() && !Validator.Pin()->GetHasPassed())
+						{
+							bAllSucceedIncludingOptional = false;
+						}
+					}
+
+
+					if(bAllSucceedIncludingOptional)
+					{
+						UE_LOG(LogSubmitTool, Log, TEXT("User has opted in to automatically submit on validation successful. Proceeding with submission..."));
+						StartSubmitProcess();
+					}
+					else
+					{
+						UE_LOG(LogSubmitTool, Warning, TEXT("User has opted in to automatically submit on validation successful but not all validations succeeded. Fix optional validation errors or submit manually if you want to bypass them."));
+						FDialogFactory::ShowInformationDialog(FText::FromString(TEXT("Auto-Submit Cancelled")), FText::FromString(TEXT("Submit tool couldn't auto submit because there were optional validations that failed.\n\nFix these errors or manually submit if you are certain that you should ignore them.")));
+					}
+				}
+			}
+
+		}));
+
+	OnPresubmitFinishedHandle = PresubmitOperationsService->OnTasksQueueFinished.Add(FOnTaskFinished::FDelegate::CreateRaw(this, &FModelInterface::OnPresubmitOperationsComplete));
+
+	FTSTicker::GetCoreTicker().AddTicker(FTickerDelegate::CreateRaw(this, &FModelInterface::Tick));
+}
+
+FModelInterface::~FModelInterface()
+{
+	PrepareSubmitCallBack.Clear();
+	FileRefreshedCallback.Clear();
+	ValidationService->OnTasksRunResultUpdated.Remove(OnValidationStateUpdatedHandle);
+	ValidationService->OnTasksQueueFinished.Remove(OnValidationFinishedHandle);
+	ValidationService->OnSingleTaskFinished.Remove(OnSingleValidationFinishedHandle);
+	PresubmitOperationsService->OnTasksQueueFinished.Remove(OnPresubmitFinishedHandle);
+}
+
+void FModelInterface::Dispose() const
+{
+	ChangelistService->CancelP4Operations();
+
+	if(GetInputEnabled())
+	{
+		ChangelistService->SendCLDescriptionToP4(EConcurrency::Synchronous);
+	}
+
+	for(const TPair<FString, TSharedPtr<FIntegrationOptionBase>>& IntegrationOption : ServiceProvider->GetService<FIntegrationService>()->GetIntegrationOptions())
+	{
+		FString Value;
+		if(IntegrationOption.Value->GetJiraValue(Value) && !Value.IsEmpty())
+		{
+			ServiceProvider->GetService<ICacheDataService>()->SetIntegrationFieldValue(GetCLID(), IntegrationOption.Key, Value);
+		}
+	}
+
+	ServiceProvider->GetService<ICacheDataService>()->SaveCacheToDisk();
+}
+
+void FModelInterface::ParseValidators() const
+{
+	TArray<TSharedRef<FValidatorBase>> Tasks;
+
+	for(const TPair<FName, FString> DefinitionPair : Parameters.Validators)
+	{
+		FValidatorDefinition TaskDefinition;
+		FStringOutputDevice Errors;
+		FValidatorDefinition::StaticStruct()->ImportText(*DefinitionPair.Value, &TaskDefinition, nullptr, 0, &Errors, FValidatorDefinition::StaticStruct()->GetName());
+
+		if(!Errors.IsEmpty())
+		{
+			UE_LOG(LogSubmitTool, Error, TEXT("Error loading parameter file %s"), *Errors);
+			FModelInterface::SetErrorState();
+			continue;
+		}
+
+		if(TaskDefinition.Type.TrimStartAndEnd().IsEmpty())
+		{
+			UE_LOG(LogSubmitTool, Error, TEXT("Task %s didn't have a Type."), *DefinitionPair.Key.ToString());
+			continue;
+		}
+
+		if(TaskDefinition.Type.Equals(SubmitToolParseConstants::TagValidator, ESearchCase::IgnoreCase))
+		{
+			Tasks.Add(MakeShared<FTagValidator>(DefinitionPair.Key, Parameters, ServiceProvider.ToSharedRef(), DefinitionPair.Value));
+		}
+		else if(TaskDefinition.Type.Equals(SubmitToolParseConstants::UBTValidator, ESearchCase::IgnoreCase))
+		{
+			Tasks.Add(MakeShared<FUBTValidator>(DefinitionPair.Key, Parameters, ServiceProvider.ToSharedRef(), DefinitionPair.Value));
+		}
+		else if(TaskDefinition.Type.Equals(SubmitToolParseConstants::CustomValidator, ESearchCase::IgnoreCase))
+		{
+			Tasks.Add(MakeShared<FValidatorRunExecutable>(DefinitionPair.Key, Parameters, ServiceProvider.ToSharedRef(), DefinitionPair.Value));
+		}
+		else if(TaskDefinition.Type.Equals(SubmitToolParseConstants::CrossChangelistValidator, ESearchCase::IgnoreCase))
+		{
+			Tasks.Add(MakeShared<FCrossChangelistValidator>(DefinitionPair.Key, Parameters, ServiceProvider.ToSharedRef(), DefinitionPair.Value));
+		}
+		else if(TaskDefinition.Type.Equals(SubmitToolParseConstants::PreflightValidator, ESearchCase::IgnoreCase))
+		{
+			Tasks.Add(MakeShared<FPreflightValidator>(DefinitionPair.Key, Parameters, ServiceProvider.ToSharedRef(), DefinitionPair.Value));
+		}
+		else
+		{
+			UE_LOG(LogSubmitTool, Error, TEXT("[%s] is not a recognized validator type and has not been activated."), *DefinitionPair.Key.ToString());
+		}
+	}
+
+	ValidationService->InitializeTasks(Tasks);
+}
+
+void FModelInterface::ParsePreSubmitOperations() const
+{
+	TArray<TSharedRef<FValidatorBase>> Tasks;
+
+	for(const TPair<FName, FString> DefinitionPair : Parameters.PresubmitOperations)
+	{
+		FValidatorDefinition TaskDefinition;
+		FStringOutputDevice Errors;
+		FValidatorDefinition::StaticStruct()->ImportText(*DefinitionPair.Value, &TaskDefinition, nullptr, 0, &Errors, FValidatorDefinition::StaticStruct()->GetName());
+
+		if(!Errors.IsEmpty())
+		{
+			UE_LOG(LogSubmitTool, Error, TEXT("Error loading parameter file %s"), *Errors);
+			FModelInterface::SetErrorState();
+			continue;
+		}
+
+		if(TaskDefinition.Type.TrimStartAndEnd().IsEmpty())
+		{
+			UE_LOG(LogSubmitTool, Error, TEXT("Task %s didn't have a Type."), *DefinitionPair.Key.ToString());
+			continue;
+		}
+
+		if(TaskDefinition.Type.Equals(SubmitToolParseConstants::CustomValidator, ESearchCase::IgnoreCase))
+		{
+			Tasks.Add(MakeShared<FValidatorRunExecutable>(DefinitionPair.Key, Parameters, ServiceProvider.ToSharedRef(), DefinitionPair.Value));
+		}
+		else if(TaskDefinition.Type.Equals(SubmitToolParseConstants::VirtualizationToolOp, ESearchCase::IgnoreCase))
+		{
+			Tasks.Add(MakeShared<FVirtualizationOperation>(DefinitionPair.Key, Parameters, ServiceProvider.ToSharedRef(), DefinitionPair.Value));
+		}
+		else
+		{
+			UE_LOG(LogSubmitTool, Error, TEXT("[%s] is not a recognized pre submit operation type and has not been activated."), *DefinitionPair.Key.ToString());
+		}
+	}
+
+	PresubmitOperationsService->InitializeTasks(Tasks);
+}
+
+
+
+void FModelInterface::SetCLDescription(const FText& newDescription, bool DoNotInvalidate) const
+{
+	if(ChangelistService->SetCLDescription(newDescription.ToString()))
+	{
+		TagService->ParseCLDescription();
+
+		if(!DoNotInvalidate)
+		{
+			ValidationService->InvalidateForChanges(ETaskArea::Changelist);
+		}
+	}
+}
+
+void FModelInterface::SendDescriptionToP4() const
+{
+	if(GetInputEnabled())
+	{
+		if(IsP4OperationRunning())
+		{
+			UE_LOG(LogSubmitToolP4, Log, TEXT("Attempted to send description to P4, but another operation is already running"));
+			return;
+		}
+
+		ChangelistService->SendCLDescriptionToP4();
+	}
+}
+
+bool FModelInterface::CanLaunchPreflight() const
+{
+	// Check Validators which are validating files, ignore changelist (description, valid tags) validators when we evaluate if we 
+	// allow the user to trigger a preflight
+	for(const TWeakPtr<const FValidatorBase>& Validator : ValidationService->GetTasks())
+	{
+		if(Validator.IsValid())
+		{
+			const TSharedPtr<const FValidatorBase> Pinned = Validator.Pin();
+			if((Pinned->Definition->TaskArea & ETaskArea::ShelveAndLocalFiles) != ETaskArea::None)
+			{
+				if(!Pinned->GetHasPassed())
+				{
+					return false;
+				}
+			}
+		}
+	}
+
+	return true;
+}
+
+void FModelInterface::ReevaluateSubmitToolTag()
+{
+	UpdateSubmitToolTag(ValidationService->GetIsRunSuccessful());	
+}
+
+void FModelInterface::UpdateSubmitToolTag(bool InbAdd)
+{
+	// add a special tag to the CL description
+	FString SubmitToolTag = FString::Format(TEXT("#submittool {0}\n"), { FAppVersion::GetVersion() });
+	FString DescriptionCopy = ChangelistService->GetCLDescription();
+
+	if(InbAdd)
+	{
+		if(!HasSubmitToolTag())
+		{
+			UE_LOG(LogSubmitToolDebug, Log, TEXT("Added Submit Tool tag"));
+			DescriptionCopy.Append(TEXT("\n") + SubmitToolTag);
+			ChangelistService->SetCLDescription(DescriptionCopy, true);
+		}
+	}
+	else if(HasSubmitToolTag())
+	{
+		FString VersionlessTag = TEXT("#submittool ");
+		int32 Loc = DescriptionCopy.Find(VersionlessTag);
+		if(Loc >= 0)
+		{
+			size_t EndPos = Loc + VersionlessTag.Len();
+			while(EndPos < DescriptionCopy.Len())
+			{
+				if(DescriptionCopy[EndPos] == TCHAR('\n'))
+				{
+					++EndPos;
+					break;
+				}
+				++EndPos;
+			}
+
+			if(Loc != 0 && DescriptionCopy[Loc - 1] == '\n' && (EndPos - DescriptionCopy.Len()) < 2)
+			{
+				Loc--;
+			}
+
+			UE_LOG(LogSubmitToolDebug, Log, TEXT("Removed Submit Tool tag"));
+			DescriptionCopy.RemoveAt(Loc, EndPos - Loc);
+			ChangelistService->SetCLDescription(DescriptionCopy, true);
+			TagService->ParseCLDescription();
+		}
+	}
+}
+
+
+bool FModelInterface::HasSubmitToolTag() const
+{
+	// Only Checking that it has a submit tool tag, regardless of version.
+	return ChangelistService->GetCLDescription().Find(TEXT("#submittool ")) != INDEX_NONE;
+}
+
+
+void FModelInterface::UpdateCLFromP4Async() const
+{
+	if(GetInputEnabled() || SubmitToolState == ESubmitToolAppState::Errored || SubmitToolState == ESubmitToolAppState::SubmitLocked)
+	{
+		ChangelistService->FetchChangelistDataAsync();
+	}
+}
+
+bool FModelInterface::GetInputEnabled()
+{
+	return SubmitToolState == ESubmitToolAppState::WaitingUserInput || SubmitToolState == ESubmitToolAppState::SubmitLocked;
+}
+
+void FModelInterface::RequestPreflight()
+{
+	if(CanLaunchPreflight())
+	{
+		PreflightService->RequestPreflight();
+	}
+	else
+	{
+		bPreflightQueued = true;
+	}
+}
+
+void FModelInterface::ShowSwarmReview()
+{
+	if (HasSwarmReview() && SwarmService.IsValid())
+	{
+		FString Url;
+		if (SwarmService->GetCurrentReviewUrl(Url))
+		{
+			UE_LOG(LogSubmitTool, Log, TEXT("Swarm: Opening Swarm Review with URL: \"%s\""), *Url);
+
+			FPlatformProcess::LaunchURL(*Url, nullptr, nullptr);
+		}
+	}
+}
+
+void FModelInterface::RequestSwarmReview()
+{
+	if (!HasSwarmReview() && SwarmService.IsValid())
+	{
+		if (!HasShelvedFiles())
+		{
+			ChangelistService->CreateShelvedFiles(FSourceControlOperationComplete::CreateLambda([this](const FSourceControlOperationRef& DeleteShelvedOp, ECommandResult::Type Result)
+				{
+					if (Result == ECommandResult::Succeeded)
+					{
+						RequestSwarmReview();
+					}
+					else
+					{
+						UE_LOG(LogSubmitTool, Error, TEXT("Failed to shelve files, Swarm Review request is cancelled"));
+					}
+				}));
+
+			return;
+		}
+
+		SwarmService->CreateReview(OnCreateReviewComplete::CreateRaw(this, &FModelInterface::OnSwarmCreateCompleted));
+	}
+}
+
+void FModelInterface::StartSubmitProcess()
+{
+	PresubmitOperationsService->ResetStates();
+
+	// Check if any last minute file changes have come in that invalidated any validators.
+	CheckForFileEdits();
+	if(IsCLValid())
+	{
+		if(PrepareSubmitCallBack.IsBound())
+		{
+			PrepareSubmitCallBack.Broadcast();
+		}
+
+		UpdateSubmitToolTag(true);
+
+		if(Parameters.IncompatibleFilesParams.IncompatibleFileGroups.Num() > 0)
+		{
+			const TArray<FString>& FilesInCL = ChangelistService->GetFilesDepotPaths();
+
+			for(const FIncompatibleFilesGroup& FileGroup : Parameters.IncompatibleFilesParams.IncompatibleFileGroups)
+			{
+				TArray<size_t, TInlineAllocator<8>> Indexes;
+
+				for(const FString& File : FilesInCL)
+				{
+					for(size_t i = 0; i < FileGroup.FileGroups.Num(); ++i)
+					{
+						if(File.Contains(FConfiguration::Substitute(FileGroup.FileGroups[i]), ESearchCase::IgnoreCase))
+						{
+							if(!Indexes.Contains(i))
+							{
+								Indexes.Add(i);
+							}
+							break;
+						}
+					}
+				}
+
+				if(Indexes.Num() > 1)
+				{
+					const FText TextTitle = FText::FromString(FileGroup.Title);
+					const FText TextDescription = FText::FromString(FileGroup.GetMessage());
+
+					if(FileGroup.bIsError)
+					{
+						FDialogFactory::ShowInformationDialog(TextTitle, TextDescription);
+						UE_LOG(LogSubmitTool, Log, TEXT("Submission canceled due to incompatible files"));
+						return;
+					}
+					else
+					{
+						if(FDialogFactory::ShowConfirmDialog(TextTitle, TextDescription) != EDialogFactoryResult::Confirm)
+						{
+							UE_LOG(LogSubmitTool, Log, TEXT("Submission canceled by user"));
+							return;
+						}
+					}
+				}				
+			}
+		}
+
+		ChangeState(ESubmitToolAppState::Submitting, SubmitToolState == ESubmitToolAppState::SubmitLocked && bIsUserInAllowlist);
+		
+
+		if(HasShelvedFiles())
+		{
+			DeleteShelvedFiles();
+		}
+		else
+		{
+			RevertUnchangedAndSubmit();
+		}
+	}
+	else
+	{
+		UE_LOG(LogSubmitTool, Warning, TEXT("Attempted to submit, but all validators have not passed. Aborting submit."));
+	}
+}
+
+void FModelInterface::RequestIntegration() const
+{
+	FNIntegrationService->RequestIntegration(FOnBooleanValueChanged::CreateLambda([this](bool bSuccess)
+		{
+			if(bSuccess)
+			{
+				ChangeState(ESubmitToolAppState::Finished);
+			}
+		}));
+}
+void FModelInterface::RefreshStateBasedOnFiles()
+{
+	TArray<FSourceControlStateRef> LocalFiles = ChangelistService->GetFilesInCL();
+
+	if(LocalFiles.IsEmpty())
+	{
+		const TArray<FSourceControlStateRef>& ShelvedFiles = ChangelistService->GetShelvedFilesInCL();
+		if(ShelvedFiles.IsEmpty())
+		{
+			UE_LOG(LogSubmitTool, Error, TEXT("There are no files in CL %s, SUBMIT IS DISABLED"), *ChangelistService->GetCLID());
+			ChangeState(ESubmitToolAppState::Errored);
+		}
+		else
+		{
+			if(!HasSubmitToolTag())
+			{
+				UE_LOG(LogSubmitTool, Warning, TEXT("This CL hasn't been validated and there are no local files. You need to unshelve and run validations."), *ChangelistService->GetCLID());
+				ChangeState(ESubmitToolAppState::Errored);
+			}
+			else
+			{
+				UE::Tasks::Launch(UE_SOURCE_LOCATION, [this] {
+					if(P4LockdownService->ArePathsInLockdown(ChangelistService->GetShelvedFilesDepotPaths(), bIsUserInAllowlist))
+					{
+						UE_LOG(LogSubmitTool, Log, TEXT("There are no local files in CL %s, Submit is disabled but you can still request an Integration with your shelved files"), *ChangelistService->GetCLID());
+						ChangeState(ESubmitToolAppState::SubmitLocked);
+					}
+					else
+					{
+						UE_LOG(LogSubmitTool, Error, TEXT("There are no files in CL %s, SUBMIT IS DISABLED"), *ChangelistService->GetCLID());
+						ChangeState(ESubmitToolAppState::Errored);
+					}
+					});
+			}
+		}
+	}
+	else
+	{
+		const TArray<FSCCStream*>& Streams = SourceControlService->GetClientStreams();
+		if(!Streams.IsEmpty())
+		{
+			FString StreamsMsg = FString::JoinBy(Streams, TEXT(" -> "), [](const FSCCStream* InStr) { return InStr->Name; });
+			for(const FString& File : ChangelistService->GetFilesDepotPaths())
+			{
+				bool bMappedToView = false;
+
+				for(const FSCCStream* Str : Streams)
+				{
+					if(File.StartsWith(Str->Name))
+					{
+						bMappedToView = true;
+						break;						
+					}
+
+					for(const FString& ImportStream : Str->AdditionalImportPaths)
+					{
+						if(File.StartsWith(ImportStream))
+						{
+							bMappedToView = true;
+							break;
+						}
+					}
+				}
+
+				if(!bMappedToView)
+				{
+					UE_LOG(LogSubmitTool, Warning, TEXT("File %s is not in the stream that the workspace is set to: %s"), *File, *StreamsMsg);
+				}
+			}
+		}
+
+		UE::Tasks::Launch(UE_SOURCE_LOCATION, [this] {
+			if(P4LockdownService->ArePathsInLockdown(ChangelistService->GetFilesDepotPaths(), bIsUserInAllowlist))
+			{
+				ChangeState(ESubmitToolAppState::SubmitLocked);
+			}
+			else
+			{
+				ChangeState(ESubmitToolAppState::WaitingUserInput);
+			}
+			});
+	}
+}
+
+void FModelInterface::OnChangelistReady(bool bIsValid)
+{
+	if(SubmitToolState == ESubmitToolAppState::Initializing)
+	{
+		if(bIsValid)
+		{
+			UE_LOG(LogSubmitTool, Log, TEXT("Retrieved information for CL %s"), *ChangelistService->GetCLID());
+			PreflightService->FetchPreflightInfo(true);
+
+			TagService->ParseCLDescription();
+			SwarmService->FetchReview(OnGetReviewComplete::CreateRaw(this, &FModelInterface::OnGetUsersFromSwarmCompleted));
+
+			RefreshStateBasedOnFiles();
+			if(!ChangelistService->GetFilesInCL().IsEmpty())
+			{
+				UpdateSubmitToolTag(false);
+				ValidationService->CheckForTagSkips();
+
+				ETaskArea ValidateArea = ~ETaskArea::Changelist;
+				for(const FTag* Tag : TagService->GetTagsArray())
+				{
+					if(Tag->GetValues().Num() != 0)
+					{
+						ValidateArea = ETaskArea::Everything;
+						break;
+					}
+				}
+
+				ValidationService->QueueByArea(ValidateArea);
+			}
+
+			FileRefreshedCallback.Broadcast();
+		}
+		else
+		{
+			UE_LOG(LogSubmitTool, Error, TEXT("Couldn't retrieve information for CL %s"), *ChangelistService->GetCLID());
+			ChangeState(ESubmitToolAppState::Errored);
+		}
+	}
+}
+
+
+void FModelInterface::RevertUnchangedAndSubmit()
+{
+	bool bHasEditOrAdd = false;
+	for(const FSourceControlStateRef& file : ChangelistService->GetFilesInCL())
+	{
+		bHasEditOrAdd |= file->IsAdded() || file->IsCheckedOut();
+	}
+
+	if(bHasEditOrAdd)
+	{
+		ChangelistService->RevertUnchangedFilesAsync(RevertUnchangedCallback);
+	}
+	else
+	{
+		Submit();
+	}
+}
+
+void FModelInterface::Submit()
+{
+	if(PresubmitOperationsService->AreTasksPendingQueue())
+	{
+		// This will call submit again when it's done
+		PresubmitOperationsService->CheckForTagSkips();
+		if(PresubmitOperationsService->QueueAll())
+		{
+			return;
+		}
+	}
+
+	auto AddendumAccumulator = [](const TArray<FString>& InAddendums, const FString& InDescription, FString& InOutAccumulated) {
+		for(const FString& Str : InAddendums)
+		{
+			if(!InDescription.Contains(Str, ESearchCase::IgnoreCase))
+			{
+				InOutAccumulated += (TEXT("\n") + Str);
+			}
+		}
+	};
+
+	const FString& CLDescription = GetCLDescription();
+	FString Addendums;
+	AddendumAccumulator(ValidationService->GetAddendums(), CLDescription, Addendums);
+	AddendumAccumulator(PresubmitOperationsService->GetAddendums(), CLDescription, Addendums);
+
+	ChangelistService->Submit(Addendums, SubmitFinishedCallback);
+}
+
+void FModelInterface::OnDeleteShelveOperationComplete(const FSourceControlOperationRef& Operation, ECommandResult::Type Result)
+{
+	if(SubmitToolState == ESubmitToolAppState::Submitting)
+	{
+		if(Result == ECommandResult::Succeeded)
+		{
+			RevertUnchangedAndSubmit();
+		}
+		else if(Result == ECommandResult::Failed)
+		{
+			ChangeState(ESubmitToolAppState::WaitingUserInput);
+		}
+		else if(Result == ECommandResult::Cancelled)
+		{
+			ChangeState(ESubmitToolAppState::WaitingUserInput);
+		}
+	}
+	else
+	{
+		ChangeState(ESubmitToolAppState::WaitingUserInput);
+	}
+}
+
+void FModelInterface::OnRevertUnchangedOperationComplete(const FSourceControlOperationRef& Operation, ECommandResult::Type Result)
+{
+	if(SubmitToolState == ESubmitToolAppState::Submitting)
+	{
+		// Revert Unchanged returns as failed if there were no files to revert, check ErrorMessages to see actual failures
+		if(Result == ECommandResult::Cancelled)
+		{
+			ChangeState(ESubmitToolAppState::WaitingUserInput);
+		}
+		else if(Result == ECommandResult::Succeeded || Operation->GetResultInfo().ErrorMessages.Num() == 0)
+		{
+			Submit();
+		}
+		else if(Result == ECommandResult::Failed)
+		{
+			ChangeState(ESubmitToolAppState::WaitingUserInput);
+		}
+	}
+	else
+	{
+		ChangeState(ESubmitToolAppState::WaitingUserInput);
+	}
+}
+
+
+void FModelInterface::OnPresubmitOperationsComplete(bool bInSuccess)
+{
+	if(bInSuccess)
+	{
+		Submit();
+	}
+	else
+	{
+		UE_LOG(LogSubmitTool, Warning, TEXT("Presubmit operations have failed, submission is not possible, please fix errors and try again."));
+		PresubmitOperationsService->ResetStates();
+	}
+}
+
+
+void FModelInterface::OnSubmitOperationComplete(const FSourceControlOperationRef& Operation, ECommandResult::Type Result)
+{
+	if(Result == ECommandResult::Succeeded)
+	{
+		// We've submitted, or tried to submit and failed so we only let the user close the app
+		ChangeState(ESubmitToolAppState::Finished);
+
+		if(FSubmitToolUserPrefs::Get()->bCloseOnSubmit)
+		{
+			MainTab.Pin()->RequestCloseTab();
+		}
+	}
+	else
+	{
+		ChangeState(ESubmitToolAppState::WaitingUserInput);
+	}
+}
+
+bool FModelInterface::Tick(float InDeltaTime)
+{
+	switch(SubmitToolState)
+	{
+	case ESubmitToolAppState::WaitingUserInput:
+		if(SwarmService->IsRequestRunning() || JiraService->IsBlockingRequestRunning())
+		{
+			ChangeState(ESubmitToolAppState::P4BlockingOperation);
+		}
+		break;
+
+	case ESubmitToolAppState::P4BlockingOperation:
+		if(!SwarmService->IsRequestRunning() && !ChangelistService->IsP4OperationRunning() && !JiraService->IsBlockingRequestRunning())
+		{
+			ChangeState(ESubmitToolAppState::WaitingUserInput);
+		}
+		break;
+
+	default:
+		break;
+	}
+
+	return true;
+}
+
+void FModelInterface::ChangeState(ESubmitToolAppState newState, bool bForce)
+{
+	if(bForce)
+	{
+		UE_LOG(LogSubmitToolDebug, Log, TEXT("Transitioned state from '%s' to '%s'"), *StaticEnum<ESubmitToolAppState>()->GetNameStringByValue(static_cast<int64>(SubmitToolState)), *StaticEnum<ESubmitToolAppState>()->GetNameStringByValue(static_cast<int64>(newState)));
+		SubmitToolState = newState;
+	}
+	else
+	{
+		if(SubmitToolAppState::AllowedTransitions.Contains(SubmitToolState))
+		{
+			const TArray<ESubmitToolAppState>& allowedStates = SubmitToolAppState::AllowedTransitions[SubmitToolState];
+			if(allowedStates.Contains(newState))
+			{
+				UE_LOG(LogSubmitToolDebug, Log, TEXT("Transitioned state from '%s' to '%s'"), *StaticEnum<ESubmitToolAppState>()->GetNameStringByValue(static_cast<int64>(SubmitToolState)), *StaticEnum<ESubmitToolAppState>()->GetNameStringByValue(static_cast<int64>(newState)));
+				SubmitToolState = newState;
+			}
+			else
+			{
+				UE_LOG(LogSubmitToolDebug, Warning, TEXT("Invalid state transition requested from '%s' to '%s'"), *StaticEnum<ESubmitToolAppState>()->GetNameStringByValue(static_cast<int64>(SubmitToolState)), *StaticEnum<ESubmitToolAppState>()->GetNameStringByValue(static_cast<int64>(newState)));
+			}
+		}
+		else
+		{
+			UE_LOG(LogSubmitToolDebug, Warning, TEXT("Transition not allowed from '%s' to '%s'"), *StaticEnum<ESubmitToolAppState>()->GetNameStringByValue(static_cast<int64>(SubmitToolState)), *StaticEnum<ESubmitToolAppState>()->GetNameStringByValue(static_cast<int64>(newState)));
+		}
+	}	
+}
+
+void FModelInterface::OnGetUsersFromSwarmCompleted(const TUniquePtr<FSwarmReview>& InReview, const FString& InErrorMessage)
+{
+	if(!InReview.IsValid())
+	{
+		UE_LOG(LogSubmitTool, Log, TEXT("Could not retrieve swarm review for current changelist. %s"), *InErrorMessage);
+		return;
+	}
+
+	TArray<const FTag*> TargetTags;
+
+	for(const FTag* Tag : TagService->GetTagsArray())
+	{
+		if(Tag->Definition.InputSubType.Equals(TEXT("Swarm"), ESearchCase::IgnoreCase))
+		{
+			TargetTags.Add(Tag);
+		}
+	}
+
+	if(!TargetTags.IsEmpty())
+	{
+		TArray<FString> SwarmUserValues;
+
+		for(const TPair<FString, FSwarmReviewParticipant>& Participant : InReview->Participants)
+		{
+			if(Participant.Key.Equals(InReview->Author, ESearchCase::IgnoreCase))
+			{
+				continue;
+			}
+
+			if(Participant.Value.Vote.Value == 1)
+			{
+				if(!SwarmUserValues.Contains(Participant.Key) && !SwarmUserValues.Contains(TEXT("@") + Participant.Key))
+				{
+					SwarmUserValues.Add(Participant.Key);
+				}
+			}
+		}
+
+		if(!SwarmUserValues.IsEmpty())
+		{
+			bool bApplied = false;
+
+			for(const FTag* Tag : TargetTags)
+			{
+				if(Tag->GetValues() != SwarmUserValues)
+				{
+					SetTagValues(*Tag, SwarmUserValues);
+					bApplied = true;
+				}
+			}
+
+			if(bApplied)
+			{
+				UE_LOG(LogSubmitTool, Log, TEXT("RB tag set to users that upvoted review '%d' Users: %s"), InReview->Id, *FString::Join(SwarmUserValues, TEXT(", ")));
+				UE_LOG(LogSubmitToolDebug, Log, TEXT("Re-running Tag validator after applying the #rb from swarm"));
+				ValidateCLDescription();
+			}
+		}
+	}
+}
+
+void FModelInterface::OnSwarmCreateCompleted(bool InResult, const FString& InErrorMessage)
+{
+	if(InResult)
+	{
+		OnGetUsersFromSwarmCompleted(SwarmService->GetReview(), InErrorMessage);
+		ShowSwarmReview();
+	}
+}
diff --git a/Engine/Source/Programs/SubmitTool/Private/Models/ModelInterface.h b/Engine/Source/Programs/SubmitTool/Private/Models/ModelInterface.h
new file mode 100644
--- /dev/null
+++ b/Engine/Source/Programs/SubmitTool/Private/Models/ModelInterface.h
@@ -0,0 +1,302 @@
+// Copyright Epic Games, Inc. All Rights Reserved.
+
+#pragma once
+
+#include "CoreMinimal.h"
+#include "Logic/ChangelistService.h"
+#include "Logic/TasksService.h"
+#include "Logic/TagService.h"
+#include "Logic/Services/Interfaces/ISTSourceControlService.h"
+#include "Logic/JiraService.h"
+#include "Logic/IntegrationService.h"
+#include "Logic/SwarmService.h"
+#include "Logic/PreflightService.h"
+#include "Logic/P4LockdownService.h"
+#include "Logic/CredentialsService.h"
+#include "Logic/UpdateService.h"
+#include "Logic/Services/SubmitToolServiceProvider.h"
+#include "Models/Tag.h"
+#include "Models/PreflightData.h"
+
+#include "Parameters/SubmitToolParameters.h"
+
+#include "ModelInterface.generated.h"
+class SDockTab;
+
+UENUM()
+enum class ESubmitToolAppState : uint8 {
+	None = 0,
+	Initializing = 1,
+	WaitingUserInput = 2,
+	Errored = 3,
+	P4BlockingOperation = 4,
+	Submitting = 5,
+	SubmitLocked = 6,
+	Finished = 7,
+};
+
+namespace SubmitToolAppState {
+
+	using StateList = const TArray<ESubmitToolAppState>;
+
+	// Allowed states to transition from the original state.
+	const TMap<const ESubmitToolAppState, StateList> AllowedTransitions =
+	{
+		{ ESubmitToolAppState::Initializing, StateList {
+				ESubmitToolAppState::WaitingUserInput,
+				ESubmitToolAppState::P4BlockingOperation,
+				ESubmitToolAppState::Errored,
+				ESubmitToolAppState::SubmitLocked
+		}},
+		{ ESubmitToolAppState::WaitingUserInput, StateList {
+				ESubmitToolAppState::Submitting,
+				ESubmitToolAppState::Errored,
+				ESubmitToolAppState::P4BlockingOperation,
+				ESubmitToolAppState::SubmitLocked
+		}},
+		{ ESubmitToolAppState::Errored, StateList {
+			ESubmitToolAppState::WaitingUserInput,
+			ESubmitToolAppState::SubmitLocked
+		}},
+		{ ESubmitToolAppState::P4BlockingOperation, StateList {
+				ESubmitToolAppState::WaitingUserInput,
+				ESubmitToolAppState::Errored,
+				ESubmitToolAppState::SubmitLocked,
+				ESubmitToolAppState::Finished
+		}},
+		{ ESubmitToolAppState::Submitting, StateList {
+				ESubmitToolAppState::WaitingUserInput,
+				ESubmitToolAppState::Errored,
+				ESubmitToolAppState::SubmitLocked,
+				ESubmitToolAppState::Finished
+		}},
+		{ ESubmitToolAppState::SubmitLocked, StateList {
+				ESubmitToolAppState::WaitingUserInput,
+				ESubmitToolAppState::SubmitLocked,
+				ESubmitToolAppState::Errored,
+				ESubmitToolAppState::Finished
+		}},
+		{ ESubmitToolAppState::Finished, StateList {
+			ESubmitToolAppState::Errored
+		}}
+	};
+}
+
+DECLARE_MULTICAST_DELEGATE(FPreSubmitCallBack)
+DECLARE_MULTICAST_DELEGATE(FFilesRefresh)
+
+class FModelInterface
+{
+public:
+	FModelInterface(const FSubmitToolParameters& InParameters);
+	~FModelInterface();
+	void Dispose() const;
+	void ParseValidators() const;
+	void ParsePreSubmitOperations() const;
+
+	void SetMainTab(TSharedPtr<SDockTab> InMainTab)									{ MainTab = InMainTab; }
+	TWeakPtr<SDockTab> GetMainTab() const											{ return MainTab; }
+
+	void ApplyTag(const FString& tagID) const										{ TagService->ApplyTag(tagID);}
+	void ApplyTag(const FTag& tag) const											{ TagService->ApplyTag(const_cast<FTag&>(tag)); ValidationService->InvalidateForChanges(ETaskArea::Changelist); }
+	void RemoveTag(const FString& tagID) const										{ TagService->RemoveTag(tagID); }
+	void RemoveTag(const FTag& tag) const											{ TagService->RemoveTag(const_cast<FTag&>(tag)); ValidationService->InvalidateForChanges(ETaskArea::Changelist); }
+	void SetTagValues(const FString& tagID, const FString& values) const			{ TagService->SetTagValues(tagID, values); }
+	void SetTagValues(const FTag& tag, const FString& values) const					{ TagService->SetTagValues(const_cast<FTag&>(tag), values); ValidationService->InvalidateForChanges(ETaskArea::Changelist); }
+	void SetTagValues(const FTag& tag, const TArray<FString>& values) const			{ TagService->SetTagValues(const_cast<FTag&>(tag), values); ValidationService->InvalidateForChanges(ETaskArea::Changelist); }
+	void UpdateTagsInCL() const														{ TagService->UpdateTagsInCL(); ValidationService->InvalidateForChanges(ETaskArea::Changelist); }
+	const FTag* GetTag(const FString& tagID) const									{ return TagService->GetTag(tagID); };
+	const TArray<const FTag*>& GetTagsArray() const									{ return TagService->GetTagsArray(); }
+	void RegisterTagUpdatedCallback(const FTagUpdated::FDelegate Callback)			{ TagService->OnTagUpdated.Add(Callback); }
+
+	void SetCLDescription(const FText& newDescription, bool DoNotInvalidate = false) const;
+	void SendDescriptionToP4() const;
+	void UpdateCLFromP4Async() const;
+	const FString& GetCLDescription() const											{ return ChangelistService->GetCLDescription(); }
+	const FString GetCLID() const													{ return ChangelistService->GetCLID(); }
+
+	FOnCLDescriptionUpdated& GetCLDescriptionUpdatedDelegate()						{ return ChangelistService->OnCLDescriptionUpdated; }
+	
+	const TArray<FSourceControlStateRef>& GetFilesInCL() const						{ return ChangelistService->GetFilesInCL(); }
+	const TArray<FString>& GetDepotFilesInCL() const								{ return ChangelistService->GetFilesDepotPaths(); }
+	bool HasShelvedFiles() const													{ return ChangelistService->HasShelvedFiles(); }
+
+	bool IsP4OperationRunning(FName OperationName = FName()) const					{ return SubmitToolState == ESubmitToolAppState::P4BlockingOperation || ChangelistService->IsP4OperationRunning(OperationName); }
+	bool IsBlockingOperationRunning() const											{ return SubmitToolState == ESubmitToolAppState::P4BlockingOperation || SwarmService->IsRequestRunning() || JiraService->IsBlockingRequestRunning() || P4LockdownService->IsBlockingOperationRunning(); }
+	void CancelP4Operations(FName OperationName = FName()) const					{ ChangelistService->CancelP4Operations(OperationName); SwarmService->CancelOperations(); }
+
+	void ValidateChangelist() const													{ ValidationService->QueueAll(); }
+	void ValidateSingle(const FName& ValidatorId) const								{ ValidationService->QueueSingle(ValidatorId, true); }
+	void ValidateCLDescription() const												{ ValidationService->QueueByArea(ETaskArea::Changelist); }
+	bool IsCLValid() const															{ return ValidationService->GetIsRunSuccessful(); }
+	bool CanLaunchPreflight() const;
+	void ReevaluateSubmitToolTag();
+	void UpdateSubmitToolTag(bool InbAdd);
+	bool HasSubmitToolTag() const;
+	bool IsValidationRunning() const												{ return ValidationService->GetIsAnyTaskRunning(); }
+	const TArray<TWeakPtr<const FValidatorBase>>& GetValidators() const				{ return ValidationService->GetTasks(); }
+	void CancelValidations(const FName& InValidatorId = FName()) const				{ ValidationService->StopTasks(InValidatorId); }
+	void CheckForFileEdits() const													{ ValidationService->CheckForLocalFileEdit(); }
+
+	FDelegateHandle AddSingleValidatorFinishedCallback
+		(const FOnSingleTaskFinished::FDelegate& Callback) const					{ return ValidationService->OnSingleTaskFinished.Add(Callback); }
+
+	void RemoveSingleValidatorFinishedCallback(const FDelegateHandle& Handle) const { ValidationService->OnSingleTaskFinished.Remove(Handle); }
+
+	FDelegateHandle AddValidationFinishedCallback
+		(const FOnTaskFinished::FDelegate& Callback) const							{ return ValidationService->OnTasksQueueFinished.Add(Callback); }
+
+	void RemoveValidationFinishedCallback(const FDelegateHandle& Handle) const		{ ValidationService->OnTasksQueueFinished.Remove(Handle); }
+	
+	FDelegateHandle AddValidationUpdatedCallback
+		(const FOnTaskRunStateChanged::FDelegate& Callback) const					{ return ValidationService->OnTasksRunResultUpdated.Add(Callback); }
+	void RemoveValidationUpdatedCallback(const FDelegateHandle& Handle) const		{ ValidationService->OnTasksRunResultUpdated.Remove(Handle); }
+
+	void GetUsers(const FOnUsersGet::FDelegate& Callback) const						{ return SourceControlService->GetUsers(Callback); }
+	const TArray<TSharedPtr<FUserData>>& GetRecentUsers() const						{ return SourceControlService->GetRecentUsers(); }
+	void AddRecentUser(TSharedPtr<FUserData>& User)									{ SourceControlService->AddRecentUser(User); }
+
+	void GetGroups(const FOnGroupsGet::FDelegate& Callback) const					{ return SourceControlService->GetGroups(Callback); }
+	const TArray<TSharedPtr<FString>>& GetRecentGroups() const						{ return SourceControlService->GetRecentGroups(); }
+	void AddRecentGroup(TSharedPtr<FString>& Group)									{ SourceControlService->AddRecentGroup(Group); }
+	FString GetUsername() const														{ return CredentialsService->GetUsername(); }
+	const FString GetRootStreamName() const											{ return SourceControlService->GetRootStreamName(); }
+	const FString GetCurrentStream() const											{ return SourceControlService->GetCurrentStreamName(); }
+
+	void SetLogin(const FString& InUsername, const FString& InPassword)				{ CredentialsService->SetLogin(InUsername, InPassword); }
+
+	static bool GetInputEnabled();	
+	static void SetErrorState()														{ ChangeState(ESubmitToolAppState::Errored);}
+	const ESubmitToolAppState GetState() const										{ return SubmitToolState; }
+	bool IsSubmitBlocked() const													{ return SubmitToolState == ESubmitToolAppState::SubmitLocked; }
+	void DeleteShelvedFiles() const													{ ChangelistService->DeleteShelvedFiles(DeleteShelveCallback); }
+
+	void RequestPreflight();
+	bool IsPreflightRequestInProgress() const										{ return PreflightService->IsRequestInProgress(); }
+	bool IsPreflightQueued() const													{ return bPreflightQueued; }
+	void RefreshPreflightInformation() const										{ PreflightService->FetchPreflightInfo(); UE_LOG(LogSubmitTool, Log, TEXT("Requesting preflight information...")); }
+	const TUniquePtr<FPreflightList>& GetPreflightData()							{ return PreflightService->GetPreflightData(); }
+
+
+	void ShowSwarmReview();
+	void RequestSwarmReview();
+	FDelegateHandle AddPreflightUpdateCallback
+	(const FOnPreflightDataUpdated::FDelegate& Callback) const						{ return PreflightService->OnPreflightDataUpdated.Add(Callback); }
+
+	void RemovePreflightUpdateCallback(const FDelegateHandle& Handle) const			{ PreflightService->OnPreflightDataUpdated.Remove(Handle); }
+
+	void StartSubmitProcess();
+
+	TSharedRef<FSubmitToolServiceProvider> GetServiceProvider() { return ServiceProvider.ToSharedRef(); }
+
+	TSharedPtr<FJiraService> GetJiraService() { return this->JiraService; }
+	TSharedPtr<FSwarmService> GetSwarmService() { return this->SwarmService; }
+	TSharedPtr<FPreflightService> GetPreflightService() { return this->PreflightService; }
+
+	bool bIsUserInAllowlist = false;
+
+	const FSubmitToolParameters& GetParameters() const { return Parameters; }
+
+	bool IsIntegrationRequired() const { return GetState() == ESubmitToolAppState::SubmitLocked; };
+	void RequestIntegration() const;
+	const TMap<FString, TSharedPtr<FIntegrationOptionBase>>& GetIntegrationOptions() { return FNIntegrationService->GetIntegrationOptions(); }
+	bool ValidateIntegrationOptions(bool bSilent) const	{ return FNIntegrationService->ValidateIntegrationOptions(bSilent); }
+
+
+	FPreSubmitCallBack PrepareSubmitCallBack;
+	FFilesRefresh FileRefreshedCallback;
+	bool bSubmitOnSuccessfulValidation = false;
+
+	const TUniquePtr<FSwarmReview>& GetSwarmReview() const		{ return SwarmService->GetReview(); }
+	const void RefreshSwarmReview()								{ return SwarmService->FetchReview(OnGetReviewComplete::CreateRaw(this, &FModelInterface::OnGetUsersFromSwarmCompleted)); }
+
+	bool GetSwarmReviewUrl(FString& OutUrl) const	{ return SwarmService->GetCurrentReviewUrl(OutUrl); }
+
+	// AUTO-UPDATE
+	bool IsAutoUpdateOn() const { return Parameters.AutoUpdateParameters.bIsAutoUpdateOn; }
+	bool CheckForNewVersion() { return UpdateService->CheckForNewVersion(); }
+	FString GetDeployId() const { return UpdateService->GetDeployId(); }
+	FString GetLocalVersion() const { return UpdateService->GetLocalVersion(); }
+	FString GetLatestVersion() const { return UpdateService->GetLatestVersion(); }
+	void InstallLatestVersion() { UpdateService->InstallLatestVersion(); }
+	void CancelInstallLatestVersion() { UpdateService->Cancel(); }
+	const FString GetDownloadMessage() { return UpdateService->GetDownloadMessage(); }
+
+	const bool IsSwarmServiceValid() { return SwarmService.IsValid(); }
+	const bool HasSwarmReview()
+	{
+		if (SwarmService.IsValid())
+		{
+			const TUniquePtr<FSwarmReview>& Review = SwarmService->GetReview();
+			if(Review.IsValid())
+			{
+				return Review->Id != 0;
+			}
+		}
+		return false;
+	}
+
+private:	
+	const FSubmitToolParameters& Parameters;
+	void OnChangelistRefresh(ETaskArea InChangeType)
+	{
+		if((InChangeType & ETaskArea::Changelist) == ETaskArea::Changelist)
+		{
+			TagService->ParseCLDescription();
+		}
+
+		if((InChangeType & ETaskArea::LocalFiles) == ETaskArea::LocalFiles)
+		{
+			FileRefreshedCallback.Broadcast();		
+		}
+
+		if((InChangeType & (ETaskArea::LocalFiles | ETaskArea::ShelvedFiles)) != ETaskArea::None)
+		{
+			RefreshStateBasedOnFiles();
+		}
+
+
+		ValidationService->InvalidateForChanges(InChangeType);
+	}
+	void RefreshStateBasedOnFiles();
+	void OnChangelistReady(bool bIsValid);
+	void RevertUnchangedAndSubmit();
+	void Submit();
+	void OnDeleteShelveOperationComplete(const FSourceControlOperationRef& Operation, ECommandResult::Type Result);
+	void OnRevertUnchangedOperationComplete(const FSourceControlOperationRef& Operation, ECommandResult::Type Result);
+	void OnPresubmitOperationsComplete(bool bInSuccess);
+	void OnSubmitOperationComplete(const FSourceControlOperationRef& Operation, ECommandResult::Type Result);
+	void OnGetUsersFromSwarmCompleted(const TUniquePtr<FSwarmReview>& InReview, const FString& InErrorMessage);
+	void OnSwarmCreateCompleted(bool InResult, const FString& InErrorMessage);
+
+	FOnChangeListReadyDelegate CLReadyCallback;
+	FOnChangelistRefreshDelegate CLRefreshCallback;
+	FSourceControlOperationComplete SubmitFinishedCallback;
+	FSourceControlOperationComplete DeleteShelveCallback;
+	FSourceControlOperationComplete RevertUnchangedCallback;
+	FDelegateHandle OnValidationStateUpdatedHandle;
+	FDelegateHandle OnValidationFinishedHandle;
+	FDelegateHandle OnPresubmitFinishedHandle;
+	FDelegateHandle OnSingleValidationFinishedHandle;
+
+	TSharedPtr<ISTSourceControlService> SourceControlService;
+	TSharedPtr<FChangelistService> ChangelistService;
+	TSharedPtr<FP4LockdownService> P4LockdownService;
+	TSharedPtr<FTagService> TagService;
+	TSharedPtr<FTasksService> ValidationService;
+	TSharedPtr<FJiraService> JiraService;
+	TSharedPtr<FPreflightService> PreflightService;
+	TWeakPtr<SDockTab> MainTab;
+	TSharedPtr<FTasksService> PresubmitOperationsService;
+	TSharedPtr<FIntegrationService> FNIntegrationService;
+	TSharedPtr<FSwarmService> SwarmService;
+	TSharedPtr<FCredentialsService> CredentialsService;
+	TSharedPtr<FUpdateService> UpdateService;
+	TSharedPtr<FSubmitToolServiceProvider> ServiceProvider;
+	bool bPreflightQueued = false;
+	bool Tick(float InDeltaTime);
+
+	static ESubmitToolAppState SubmitToolState;
+
+	static void ChangeState(ESubmitToolAppState newState, bool bForce = false);
+};
\ No newline at end of file
diff --git a/Engine/Source/Programs/SubmitTool/Private/Models/Tag.cpp b/Engine/Source/Programs/SubmitTool/Private/Models/Tag.cpp
new file mode 100644
--- /dev/null
+++ b/Engine/Source/Programs/SubmitTool/Private/Models/Tag.cpp
@@ -0,0 +1,173 @@
+// Copyright Epic Games, Inc. All Rights Reserved.
+
+
+#include "Tag.h"
+#include "Misc/StringBuilder.h"
+#include "Logging/SubmitToolLog.h"
+#include "Internationalization/Regex.h"
+
+
+FTag::FTag(const FTagDefinition& def, size_t start) : Definition(def), StartPos(start)
+{
+	for(const TCHAR& c : Definition.ValueDelimiter)
+	{
+		int idx = SplittedDelims.Add(FString::Chr(c));
+		Delimiters.Add(*SplittedDelims[idx]);
+	}
+}
+
+FString FTag::GetFullTag() const
+{
+	TStringBuilder<256> strBuilder;
+
+	strBuilder.AppendChar('\n');
+	strBuilder.Append(Definition.TagId);
+
+	if(TagValues.Num() > 0)
+	{
+		strBuilder.AppendChar(' ');
+
+		for(size_t i = 0; i < TagValues.Num(); ++i)
+		{
+			if(i != 0)
+			{
+				strBuilder.Append(Definition.ValueDelimiter);
+			}
+			strBuilder.Append(TagValues[i]);
+		}
+	}
+
+	return strBuilder.ToString();
+}
+
+bool FTag::ParseTag(const FString& source)
+{
+	// regex pattern example, replacing tag, delimiter and min/maxvalues
+	// (?:(?:\r\n|\r|\n)?TAGID)( +(?:[DELIMITERS]*(?!#)(?:[\w!"\$-\/\:-\@\[-\`\{-\~]+)){MINVALUES,MAXVALUES})?
+	// (?:(?:\r\n|\r|\n)?#jira(?!\w))( +(?:[, ]*(?!#)(?:[\w!"\$-\/\:-\@\[-\`\{-\~]+)){1,256})?
+
+	FString regexPat = "(?:(?:\\r\\n|\\r|\\n)?" + Definition.GetTagId() + "(?!\\w))( +(?:[" + Definition.ValueDelimiter + "]*(?!#)(?:[\\w!\"\\$-\\/\\:-\\@\\[-\\`\\{-\\~]+)){" + FString::FromInt(Definition.MinValues) + "," + FString::FromInt(Definition.MaxValues) + "})?";
+	FRegexPattern Pattern = FRegexPattern(regexPat, ERegexPatternFlags::CaseInsensitive);
+	FRegexMatcher regex = FRegexMatcher(Pattern, source);
+	bool match = regex.FindNext();
+	if(match)
+	{
+		bIsDirty = false;
+		StartPos = regex.GetMatchBeginning();
+		LastSize = regex.GetMatchEnding() - regex.GetMatchBeginning();
+
+		UE_LOG(LogSubmitToolDebug, Log, TEXT("Start: %d"), regex.GetMatchBeginning());
+		UE_LOG(LogSubmitToolDebug, Log, TEXT("Regex matched: %s"), *regex.GetCaptureGroup(0));
+
+		FString capture = regex.GetCaptureGroup(1).TrimStart();
+		capture.ParseIntoArray(TagValues, Delimiters.GetData(), SplittedDelims.Num());
+
+		for(const FString& value : TagValues)
+		{
+			UE_LOG(LogSubmitToolDebug, Log, TEXT("Captured Value: %s"), *value);
+		}
+
+		UE_LOG(LogSubmitToolDebug, Log, TEXT("End: %d"), regex.GetMatchEnding());
+	}
+	else
+	{
+		Reset();
+		UE_LOG(LogSubmitToolDebug, Log, TEXT("Tag %s not found in description"), *Definition.GetTagId());
+	}
+
+	if(OnTagUpdated.IsBound())
+	{
+		OnTagUpdated.Broadcast(*this);
+	}
+	return match;
+}
+
+void FTag::SetValues(const FString& valuesText)
+{
+	bIsDirty = true;
+	TagValues.Empty();
+	valuesText.ParseIntoArray(TagValues, Delimiters.GetData(), SplittedDelims.Num());
+	ValidationState = ETagState::Unchecked;
+
+	if(OnTagUpdated.IsBound())
+	{
+		OnTagUpdated.Broadcast(*this);
+	}
+}
+
+FString FTag::GetValuesText() const
+{
+	TStringBuilder<256> strBuilder;
+	for(size_t i = 0; i < TagValues.Num(); ++i)
+	{
+		if(i != 0)
+		{
+			strBuilder.Append(Definition.ValueDelimiter);
+		}
+		strBuilder.Append(TagValues[i]);
+	}
+
+	return strBuilder.ToString();
+}
+
+void FTag::SetValues(const TArray<FString>& InValues)
+{
+	bIsDirty = true;
+	this->TagValues = InValues;
+
+	for(size_t i = 0; i<TagValues.Num();++i)
+	{
+		bool removedChar;
+		do
+		{
+			removedChar = false;
+			for(const FString& Delim : SplittedDelims)
+			{
+				bool localRemoved;
+				do
+				{
+					TagValues[i].TrimCharInline(Delim[0], &localRemoved);
+					removedChar |= localRemoved;
+				} while(localRemoved);
+			}
+		} while(removedChar);
+	}
+
+	if(OnTagUpdated.IsBound())
+	{
+		OnTagUpdated.Broadcast(*this);
+	}
+}
+
+const TArray<FString> FTag::GetValues(bool bEvenIfDisabled) const
+{
+	if(IsEnabled() || bEvenIfDisabled)
+	{
+		return TagValues;
+	}
+	else
+	{
+		return TArray<FString>();
+	}
+}
+
+const FTagValidationConfig& FTag::GetCurrentValidationConfig(const TArray<FString>& InDepotPaths) const
+{
+	if(Definition.ValidationOverrides.Num() != 0)
+	{
+		for(const FTagValidationOverride& ValidationOverride : Definition.ValidationOverrides)
+		{
+			FRegexPattern RegexPattern = FRegexPattern(ValidationOverride.RegexPath, ERegexPatternFlags::CaseInsensitive);
+			for(const FString& Path : InDepotPaths)
+			{
+				FRegexMatcher Regex = FRegexMatcher(RegexPattern, Path);
+				if(Regex.FindNext())
+				{
+					return ValidationOverride.ConfigOverride;
+				}
+			}
+		}
+	}
+
+	return Definition.Validation;
+}
