Moving Submit Tool to Source/Programs

diff --git a/Engine/Source/Programs/SubmitTool/Private/Parameters/SubmitToolParametersBuilder.cpp b/Engine/Source/Programs/SubmitTool/Private/Parameters/SubmitToolParametersBuilder.cpp
new file mode 100644
--- /dev/null
+++ b/Engine/Source/Programs/SubmitTool/Private/Parameters/SubmitToolParametersBuilder.cpp
@@ -0,0 +1,429 @@
+// Copyright Epic Games, Inc. All Rights Reserved.
+
+#include "SubmitToolParametersBuilder.h"
+
+#include "SubmitToolUtils.h"
+#include "Configuration/Configuration.h"
+
+#include "CoreGlobals.h"
+#include "Models/ModelInterface.h"
+#include "HAL/FileManager.h"
+#include "Logging/SubmitToolLog.h"
+#include "CommandLine/CmdLineParameters.h"
+#include "Misc/ConfigContext.h"
+
+FSubmitToolParametersBuilder::FSubmitToolParametersBuilder(const FString& InParametersXmlFile)
+{
+	ConfigHierarchy.Add({ TEXT("Base"), TEXT("{ENGINE}/Config/Base.ini") });
+	ConfigHierarchy.Add( { TEXT("SubmitToolBase"), TEXT("{PROJECT}/Config/{TYPE}.ini") } );
+	ConfigHierarchy.Add( { TEXT("Platform"), TEXT("{PROJECT}/Config/{PLATFORM}/{PLATFORM}{TYPE}.ini") } );
+
+	FString RootDir;
+	FCmdLineParameters::Get().GetValue(FSubmitToolCmdLine::RootDir, RootDir);
+	FPaths::NormalizeDirectoryName(RootDir);
+
+	// allocate a string since the FConfigLayer expects static TCHAR*, so we can't let the string go away
+	FString* RootBase = new FString(FPaths::Combine(RootDir, TEXT("/Config/{TYPE}.ini")));
+	ConfigHierarchy.Add({ TEXT("RootBase"), **RootBase, EConfigLayerFlags::NoExpand });
+
+	FString* RootPlatform = new FString(FPaths::Combine(RootDir, TEXT("/Config/{PLATFORM}/{PLATFORM}{TYPE}.ini")));
+	ConfigHierarchy.Add({ TEXT("RootPlatform"), **RootPlatform, EConfigLayerFlags::NoExpand });
+
+	if(!RootDir.IsEmpty())
+	{
+		IFileManager::Get().IterateDirectory(*FConfiguration::Substitute(TEXT("$(root)")),
+			[this](const TCHAR* FileOrDir, bool bIsDir) -> bool
+			{
+				if(bIsDir)
+				{
+					FString Dir = { FileOrDir };
+					if(Dir != TEXT("SubmitTool") && !Dir.Contains(TEXT("Engine"), ESearchCase::IgnoreCase))
+					{
+						const FString Extension = TEXT(".uproject");
+						TArray<FString> UProjects;
+						IFileManager::Get().FindFiles(UProjects, FileOrDir, *Extension);
+
+						if(UProjects.Num() != 0)
+						{
+							// allocate a string since the FConfigLayer expects static TCHAR*, so we can't let the string go away
+							FString* ProjectIni = new FString(FPaths::Combine(Dir, TEXT("/Config/{TYPE}.ini")));
+							ConfigHierarchy.Add( { TEXT("Project"), **ProjectIni, EConfigLayerFlags::NoExpand } );
+
+							FString* ProjectPlatform = new FString(FPaths::Combine(Dir, TEXT("/Config/{PLATFORM}/{PLATFORM}{TYPE}.ini")));
+							ConfigHierarchy.Add({ TEXT("ProjectPlatform"), **ProjectPlatform, EConfigLayerFlags::NoExpand });
+
+							// Let's hold on to this info for the preflight parameters
+							for (const FString& UProjectName : UProjects)
+							{
+								ProjectNames.Add(UProjectName.LeftChop(Extension.Len()));
+							}
+						}
+					}
+				}
+
+				return true;
+			});
+	}
+		
+	// allocate a string since the FConfigLayer expects static TCHAR*, so we can't let the string go away
+	FString* UserIni = new FString(FPaths::Combine(FSubmitToolUtils::GetLocalAppDataPath(), "SubmitTool", "SubmitTool.ini"));
+	ConfigHierarchy.Add( { TEXT("User"), **UserIni, EConfigLayerFlags::NoExpand } );
+}
+
+FSubmitToolParameters FSubmitToolParametersBuilder::Build()
+{
+	FConfigContext Context = FConfigContext::ReadIntoGConfig();
+	Context.OverrideLayers = ConfigHierarchy;
+
+	FString IniFilename;
+	Context.Load(TEXT("SubmitTool"), IniFilename);
+	SubmitToolConfig = GConfig->FindConfigFile(IniFilename);
+
+	UE_LOG(LogSubmitTool, Verbose, TEXT("Loading config from the following files:"));
+	for(const TPair<int32, FString>& OverrideLayer : Context.Branch->Hierarchy)
+	{
+		if(IFileManager::Get().FileExists(*OverrideLayer.Value))
+		{
+			UE_LOG(LogSubmitTool, Verbose, TEXT("%s"), *FPaths::ConvertRelativePathToFull(OverrideLayer.Value));
+		}
+	}
+
+	FSubmitToolParameters Parameters;
+	Parameters.GeneralParameters = BuildGeneralParameters();
+	Parameters.JiraParameters = BuildJiraParameters();
+	Parameters.Telemetry = GetTelemetryParameters();
+	Parameters.IntegrationParameters = BuildIntegrationParameters();
+	Parameters.AvailableTags = BuildAvailableTags();
+	Parameters.Validators = BuildValidators();
+	Parameters.PresubmitOperations = BuildPresubmitOperations();
+	Parameters.CopyLogParameters = BuildCopyLogParameters();
+	Parameters.P4LockdownParameters = BuildP4LockdownParameters();
+	Parameters.OAuthParameters = BuildOAuthParameters();
+	Parameters.IncompatibleFilesParams = BuildIncompatibleFilesParameters();
+	Parameters.HordeParameters = BuildHordeParameters();
+	Parameters.AutoUpdateParameters = BuildAutoUpdateParameters();
+	return Parameters;
+}
+
+FGeneralParameters FSubmitToolParametersBuilder::BuildGeneralParameters()
+{
+	const FConfigSection* Section = SubmitToolConfig->FindSection(TEXT("SubmitTool.General"));
+	FGeneralParameters Output;
+
+	if(Section != nullptr)
+	{
+		FStringOutputDevice Errors;
+		FGeneralParameters::StaticStruct()->ImportText(*SectionToText(*Section), &Output, nullptr, 0, &Errors, FGeneralParameters::StaticStruct()->GetName());
+
+		if(!Errors.IsEmpty())
+		{
+			UE_LOG(LogSubmitTool, Error, TEXT("Error loading parameter file %s"), *Errors);
+			FModelInterface::SetErrorState();
+		}
+		else
+		{
+			Output.CacheFile = FConfiguration::SubstituteAndNormalizeFilename(Output.CacheFile);
+		}
+	}
+
+	return Output;
+}
+
+FJiraParameters FSubmitToolParametersBuilder::BuildJiraParameters()
+{
+	const FConfigSection* Section = SubmitToolConfig->FindSection(TEXT("SubmitTool.Jira"));
+	FJiraParameters Output;
+
+	if(Section != nullptr)
+	{
+		FStringOutputDevice Errors;
+		FJiraParameters::StaticStruct()->ImportText(*SectionToText(*Section), &Output, nullptr, 0, &Errors, FJiraParameters::StaticStruct()->GetName());
+
+		if(!Errors.IsEmpty())
+		{
+			UE_LOG(LogSubmitTool, Error, TEXT("Error loading parameter file %s"), *Errors);
+			FModelInterface::SetErrorState();
+		}
+	}
+
+	return Output;
+}
+
+ FTelemetryParameters FSubmitToolParametersBuilder::GetTelemetryParameters()
+{
+	static const TCHAR* Section = TEXT("SubmitTool.Telemetry");
+	static const TCHAR* UrlKey = TEXT("URL");
+	static const TCHAR* InstanceKey = TEXT("Instance");
+
+	FTelemetryParameters Output;
+
+	FString Url;
+	if (SubmitToolConfig->GetString(Section, UrlKey, Url))
+	{
+		Output.Url = Url;
+	}
+
+	FString Instance;
+	if (SubmitToolConfig->GetString(Section, InstanceKey, Instance))
+	{
+		Output.Instance = Instance;
+	}
+
+	return Output;
+}
+
+ FIntegrationParameters FSubmitToolParametersBuilder::BuildIntegrationParameters()
+ {
+	 const FConfigSection* Section = SubmitToolConfig->FindSection(TEXT("SubmitTool.FNIntegration"));
+	 FIntegrationParameters Output;
+
+	 if(Section != nullptr)
+	 {
+		 FStringOutputDevice Errors;
+		 FIntegrationParameters::StaticStruct()->ImportText(*SectionToText(*Section), &Output, nullptr, 0, &Errors, FIntegrationParameters::StaticStruct()->GetName());
+
+		 if(!Errors.IsEmpty())
+		 {
+			 UE_LOG(LogSubmitTool, Error, TEXT("Error loading parameter file %s"), *Errors);
+			 FModelInterface::SetErrorState();
+		 }
+	 }
+
+	 return Output;
+ }
+
+TArray<FTagDefinition> FSubmitToolParametersBuilder::BuildAvailableTags()
+{
+	static const TCHAR* TagsSectionName = TEXT("Tags.");
+	TArray<FTagDefinition> Output;
+
+	for(const TPair<FString, FConfigSection>& Section : AsConst(*SubmitToolConfig))
+	{
+		if(Section.Key.StartsWith(TagsSectionName))
+		{
+			FTagDefinition Definition;
+			FStringOutputDevice Errors;
+			FTagDefinition::StaticStruct()->ImportText(*SectionToText(Section.Value), &Definition, nullptr, 0, &Errors, FValidatorDefinition::StaticStruct()->GetName());
+
+			if (Definition.bIsDisabled)
+			{
+				UE_LOG(LogSubmitToolDebug, Verbose, TEXT("Skipped tag due to it being disabled %s"), *Definition.TagId);
+				continue;
+			}
+
+			if(!Definition.DocumentationUrl.IsEmpty())
+			{
+				Definition.ToolTip += TEXT("\nClick the icon for more information.");
+			}
+
+			if(!Errors.IsEmpty())
+			{
+				UE_LOG(LogSubmitTool, Error, TEXT("Error loading parameter file %s"), *Errors);
+				FModelInterface::SetErrorState();
+			}
+			else
+			{
+				UE_LOG(LogSubmitToolDebug, Verbose, TEXT("Added Tag %s"), *Definition.TagId);
+				Output.Add(Definition);
+			}
+		}
+	}
+
+	Output.Sort([](const FTagDefinition& A, const FTagDefinition& B)
+		{
+			return A.OrdinalOverride <= B.OrdinalOverride;
+		});
+
+	return Output;
+}
+
+TMap<FString, FString> FSubmitToolParametersBuilder::BuildValidators()
+{
+	static const TCHAR* ValidatorsSectionName = TEXT("Validator.");
+
+	TMap<FString, FString> Output;
+
+	for(const TPair<FString, FConfigSection>& Section : AsConst(*SubmitToolConfig))
+	{
+		if(Section.Key.StartsWith(ValidatorsSectionName))
+		{
+			Output.Add(Section.Key.Replace(ValidatorsSectionName, TEXT("")), SectionToText(Section.Value));
+		}
+	}
+
+	return Output;
+}
+
+TMap<FString, FString> FSubmitToolParametersBuilder::BuildPresubmitOperations()
+{
+	static const TCHAR* PresubmitOperationsSectionName = TEXT("PresubmitOperation.");
+
+	TMap<FString, FString> Output;
+
+	for(const TPair<FString, FConfigSection>& Section : AsConst(*SubmitToolConfig))
+	{
+		if(Section.Key.StartsWith(PresubmitOperationsSectionName))
+		{
+			Output.Add(Section.Key.Replace(PresubmitOperationsSectionName, TEXT("")), SectionToText(Section.Value));
+		}
+	}
+
+	return Output;
+}
+
+FCopyLogParameters FSubmitToolParametersBuilder::BuildCopyLogParameters()
+{
+	const FConfigSection* Section = SubmitToolConfig->FindSection(TEXT("SubmitTool.CopyLog"));
+	FCopyLogParameters Output;
+	
+	if(Section != nullptr)
+	{
+		FStringOutputDevice Errors;
+		FCopyLogParameters::StaticStruct()->ImportText(*SectionToText(*Section), &Output, nullptr, 0, &Errors, FCopyLogParameters::StaticStruct()->GetName());
+
+		if(!Errors.IsEmpty())
+		{
+			UE_LOG(LogSubmitTool, Error, TEXT("Error loading parameter file %s"), *Errors);
+			FModelInterface::SetErrorState();
+		}
+	}
+
+	return Output;
+}
+FP4LockdownParameters FSubmitToolParametersBuilder::BuildP4LockdownParameters()
+{
+	const FConfigSection* Section = SubmitToolConfig->FindSection(TEXT("SubmitTool.P4Lockdown"));
+	FP4LockdownParameters Output;
+
+	if(Section != nullptr)
+	{
+		FStringOutputDevice Errors;
+		FP4LockdownParameters::StaticStruct()->ImportText(*SectionToText(*Section), &Output, nullptr, 0, &Errors, FP4LockdownParameters::StaticStruct()->GetName());
+
+		if(!Errors.IsEmpty())
+		{
+			UE_LOG(LogSubmitTool, Error, TEXT("Error loading parameter file %s"), *Errors);
+			FModelInterface::SetErrorState();
+		}
+	}
+
+	return Output;
+}
+
+FOAuthTokenParams FSubmitToolParametersBuilder::BuildOAuthParameters()
+{
+	const FConfigSection* Section = SubmitToolConfig->FindSection(TEXT("SubmitTool.OAuthToken"));
+	FOAuthTokenParams Output;
+
+	if(Section != nullptr)
+	{
+		FStringOutputDevice Errors;
+		FOAuthTokenParams::StaticStruct()->ImportText(*SectionToText(*Section), &Output, nullptr, 0, &Errors, FOAuthTokenParams::StaticStruct()->GetName());
+
+		if(!Errors.IsEmpty())
+		{
+			UE_LOG(LogSubmitTool, Error, TEXT("Error loading parameter file %s"), *Errors);
+			FModelInterface::SetErrorState();
+		}
+		else
+		{
+			Output.OAuthFile = FConfiguration::Substitute(Output.OAuthFile);
+			Output.OAuthTokenTool = FConfiguration::Substitute(Output.OAuthTokenTool);
+			Output.OAuthArgs = FConfiguration::Substitute(Output.OAuthArgs); 
+		}
+	}
+
+	return Output;
+}
+
+
+FIncompatibleFilesParams FSubmitToolParametersBuilder::BuildIncompatibleFilesParameters()
+{
+	const FConfigSection* Section = SubmitToolConfig->FindSection(TEXT("SubmitTool.IncompatibleFiles"));
+	FIncompatibleFilesParams Output;
+
+	if(Section != nullptr)
+	{
+		FStringOutputDevice Errors;
+		FIncompatibleFilesParams::StaticStruct()->ImportText(*SectionToText(*Section), &Output, nullptr, 0, &Errors, FIncompatibleFilesParams::StaticStruct()->GetName());
+
+		if(!Errors.IsEmpty())
+		{
+			UE_LOG(LogSubmitTool, Error, TEXT("Error loading parameter file %s"), *Errors);
+			FModelInterface::SetErrorState();
+		}
+	}
+
+	return Output;
+}
+
+FHordeParameters FSubmitToolParametersBuilder::BuildHordeParameters()
+{
+	const FConfigSection* Section = SubmitToolConfig->FindSection(TEXT("SubmitTool.Horde"));
+	FHordeParameters Output;
+
+	if(Section != nullptr)
+	{
+		FStringOutputDevice Errors;
+		FHordeParameters::StaticStruct()->ImportText(*SectionToText(*Section), &Output, nullptr, 0, &Errors, FHordeParameters::StaticStruct()->GetName());
+
+		if (!Errors.IsEmpty())
+		{
+			UE_LOG(LogSubmitTool, Error, TEXT("Error loading parameter file %s"), *Errors);
+			FModelInterface::SetErrorState();
+		}
+	}
+	return Output;
+}
+
+FAutoUpdateParameters FSubmitToolParametersBuilder::BuildAutoUpdateParameters()
+{
+	const FConfigSection* Section = SubmitToolConfig->FindSection(TEXT("SubmitTool.AutoUpdate"));
+	FAutoUpdateParameters Output;
+
+	if(Section != nullptr)
+	{
+		FStringOutputDevice Errors;
+		FAutoUpdateParameters::StaticStruct()->ImportText(*SectionToText(*Section), &Output, nullptr, 0, &Errors, FAutoUpdateParameters::StaticStruct()->GetName());
+
+		if (!Errors.IsEmpty())
+		{
+			UE_LOG(LogSubmitTool, Error, TEXT("Error loading parameter file %s"), *Errors);
+			FModelInterface::SetErrorState();
+		}
+		else
+		{
+			Output.DeployIdFilePath = FConfiguration::Substitute(Output.DeployIdFilePath);
+			Output.LocalDownloadZip = FConfiguration::Substitute(Output.LocalDownloadZip);
+			Output.LocalVersionFile = FConfiguration::Substitute(Output.LocalVersionFile);
+			Output.AutoUpdateScript = FConfiguration::Substitute(Output.AutoUpdateScript);
+			Output.LocalAutoUpdateScript = FConfiguration::Substitute(Output.LocalAutoUpdateScript);
+		}
+
+	}
+
+	return Output;
+}
+
+FString FSubmitToolParametersBuilder::SectionToText(const FConfigSection& InSection) const
+{
+	TArray<FString> lines;
+	for(const TPair<FName, FConfigValue>& Item : InSection.Array())
+	{
+		FString Value = Item.Value.GetValue();
+
+		// If it's an array/map/struct, we only need to quote the key, otherwise quote key and value
+		if((Value.IsNumeric() && !Value.Equals(TEXT("-"))) || (Value.StartsWith(TEXT("(")) && Value.EndsWith(TEXT(")")) && !Item.Key.ToString().Contains(TEXT("Regex"), ESearchCase::IgnoreCase)))
+		{
+			lines.Add(FString::Printf(TEXT("\"%s\"=%s"), *Item.Key.ToString(), *Item.Value.GetValue()));
+		}
+		else
+		{
+			lines.Add(FString::Printf(TEXT("\"%s\"=\"%s\""), *Item.Key.ToString(), *Item.Value.GetValue()));
+		}
+	}
+
+	FString FinalText = TEXT("(") + FString::Join(lines, TEXT(",")) + TEXT(")");
+	return FinalText;
+}
+
diff --git a/Engine/Source/Programs/SubmitTool/Private/Parameters/SubmitToolParametersBuilder.h b/Engine/Source/Programs/SubmitTool/Private/Parameters/SubmitToolParametersBuilder.h
new file mode 100644
--- /dev/null
+++ b/Engine/Source/Programs/SubmitTool/Private/Parameters/SubmitToolParametersBuilder.h
@@ -0,0 +1,36 @@
+// Copyright Epic Games, Inc. All Rights Reserved.
+#pragma once
+
+#include "CoreMinimal.h"
+
+#include "Misc/ConfigCacheIni.h"
+#include "SubmitToolParameters.h"
+
+class FSubmitToolParametersBuilder
+{
+public:
+	FSubmitToolParametersBuilder(const FString& InParametersXmlFile);
+	FSubmitToolParameters Build();
+
+private:
+	TArray<FConfigLayer> ConfigHierarchy;
+	TArray<FString> ProjectNames;
+	FConfigFile* SubmitToolConfig;
+private:
+	FGeneralParameters BuildGeneralParameters();
+	FJiraParameters BuildJiraParameters(); 
+	FTelemetryParameters GetTelemetryParameters();
+	FIntegrationParameters BuildIntegrationParameters();
+	TArray<FTagDefinition> BuildAvailableTags();
+	TMap<FString, FString> BuildValidators();
+	TMap<FString, FString> BuildPresubmitOperations();
+	FCopyLogParameters BuildCopyLogParameters();
+	FP4LockdownParameters BuildP4LockdownParameters();
+	FOAuthTokenParams BuildOAuthParameters();
+	FIncompatibleFilesParams BuildIncompatibleFilesParameters();
+	FHordeParameters BuildHordeParameters();
+	FAutoUpdateParameters BuildAutoUpdateParameters();
+
+	
+	FString SectionToText(const FConfigSection& InSection) const;
+};
diff --git a/Engine/Source/Programs/SubmitTool/Private/Platforms/Windows/Resources/SubmitTool.rc b/Engine/Source/Programs/SubmitTool/Private/Platforms/Windows/Resources/SubmitTool.rc
new file mode 100644
--- /dev/null
+++ b/Engine/Source/Programs/SubmitTool/Private/Platforms/Windows/Resources/SubmitTool.rc
@@ -0,0 +1,84 @@
+// Copyright Epic Games, Inc. All Rights Reserved.
+
+#define APSTUDIO_READONLY_SYMBOLS
+/////////////////////////////////////////////////////////////////////////////
+//
+// Generated from the TEXTINCLUDE 2 resource.
+//
+#include <windows.h>
+#include "Runtime/Launch/Resources/Version.h"
+#include "Runtime/Launch/Resources/Windows/PCLaunch.rc"
+#include "../../../Version/AppVersionDefines.h"
+
+#define VER_FILEVERSION         SUBMIT_TOOL_VERSION_MAJOR,SUBMIT_TOOL_VERSION_MINOR,SUBMIT_TOOL_VERSION_PATCH,0
+#define VER_FILEVERSION_STR     SUBMIT_TOOL_VERSION_STRING "\0"
+
+#define VER_PRODUCTVERSION      ENGINE_MAJOR_VERSION,ENGINE_MINOR_VERSION,ENGINE_PATCH_VERSION,0
+#define VER_PRODUCTVERSION_STR  ENGINE_VERSION_STRING "\0"
+
+#ifdef _DEBUG
+#define VER_FILEFLAGS           0x1L // VS_FF_DEBUG
+#else
+#define VER_FILEFLAGS           0x0L
+#endif
+
+
+#define IDI_SUBMIT_TOOL            IDICON_UEGame // 123
+
+/////////////////////////////////////////////////////////////////////////////
+#undef APSTUDIO_READONLY_SYMBOLS
+
+/////////////////////////////////////////////////////////////////////////////
+// English (U.S.) resources
+
+#if !defined(AFX_RESOURCE_DLL) || defined(AFX_TARG_ENU)
+#ifdef _WIN32
+LANGUAGE LANG_ENGLISH, SUBLANG_ENGLISH_US
+#pragma code_page(1252)
+#endif //_WIN32
+
+/////////////////////////////////////////////////////////////////////////////
+//
+// Version
+//
+
+VS_VERSION_INFO VERSIONINFO
+FILEVERSION     VER_FILEVERSION
+PRODUCTVERSION  VER_PRODUCTVERSION
+FILEFLAGSMASK   0x17L
+FILEFLAGS       VER_FILEFLAGS
+FILEOS          0x4L
+FILETYPE        0x2L
+FILESUBTYPE     0x0L
+BEGIN
+    BLOCK "StringFileInfo"
+    BEGIN
+        BLOCK "040904b0"
+        BEGIN
+            VALUE "CompanyName",      EPIC_COMPANY_NAME
+            VALUE "LegalCopyright",   EPIC_COPYRIGHT_STRING
+            VALUE "ProductName",      EPIC_PRODUCT_NAME
+            VALUE "ProductVersion",   VER_PRODUCTVERSION_STR
+            VALUE "FileDescription",  SUBMIT_TOOL_APPNAME
+            VALUE "FileVersion",      VER_FILEVERSION_STR
+            VALUE "InternalName",	  SUBMIT_TOOL_APPNAME
+            VALUE "OriginalFilename", "SubmitTool.exe"
+        END
+    END
+    BLOCK "VarFileInfo"
+    BEGIN
+        VALUE "Translation", 0x409, 1200
+    END
+END
+
+/////////////////////////////////////////////////////////////////////////////
+//
+// Icon
+//
+
+// Icon with lowest ID value placed first to ensure application icon
+// remains consistent on all systems.
+IDI_SUBMIT_TOOL    ICON    "SubmitTool.ico"
+
+#endif    // English (U.S.) resources
+/////////////////////////////////////////////////////////////////////////////
diff --git a/Engine/Source/Programs/SubmitTool/Private/Platforms/Windows/Resources/make_ico.py b/Engine/Source/Programs/SubmitTool/Private/Platforms/Windows/Resources/make_ico.py
new file mode 100644
--- /dev/null
+++ b/Engine/Source/Programs/SubmitTool/Private/Platforms/Windows/Resources/make_ico.py
@@ -0,0 +1,15 @@
+# Copyright Epic Games, Inc. All Rights Reserved.
+# use python to run this and generate the .ico, requires pip install pillow
+from PIL import Image
+
+def MakeIcon(folder, name, output):
+    sizes = [256, 128, 48, 32, 24, 16]
+    images = []
+    for size in sizes:
+        print(f"Loading {folder}/icon_{size}x{size}.png")
+        image = Image.open(f"{folder}/icon_{size}x{size}.png")
+        images.append(image)
+    print(f"Creating {output}")
+    images[0].save(f"{output}", append_images=images)
+
+MakeIcon("SubmitTool", "SubmitTool", "SubmitTool.ico")
diff --git a/Engine/Source/Programs/SubmitTool/Private/SubmitToolApp.cpp b/Engine/Source/Programs/SubmitTool/Private/SubmitToolApp.cpp
new file mode 100644
--- /dev/null
+++ b/Engine/Source/Programs/SubmitTool/Private/SubmitToolApp.cpp
@@ -0,0 +1,194 @@
+// Copyright Epic Games, Inc. All Rights Reserved.
+
+#include "SubmitToolApp.h"
+#include "CoreMinimal.h"
+#include "SubmitToolUtils.h"
+#include "Models/SubmitToolUserPrefs.h"
+
+#include "RequiredProgramMainCPPInclude.h"
+#include "StandaloneRenderer.h"
+#include "Framework/Application/SlateApplication.h"
+#include "Framework/Docking/TabManager.h"
+
+#include "CommandLine/CmdLineParameters.h"
+
+#include "View/SubmitToolWindow.h"
+#include "View/SubmitToolStyle.h"
+#include "View/SubmitToolCommandHandler.h"
+#include "View/SubmitToolMenu.h"
+
+#include "Models/ModelInterface.h"
+
+#include "Parameters/SubmitToolParameters.h"
+#include "Parameters/SubmitToolParametersBuilder.h"
+
+#include "Telemetry/TelemetryService.h"
+#include "Version/AppVersion.h"
+
+#include "Configuration/Configuration.h"
+
+IMPLEMENT_APPLICATION(SubmitTool, "SubmitTool");
+
+#define LOCTEXT_NAMESPACE "SubmitTool"
+
+
+int RunSubmitTool(const TCHAR* CommandLine, const FGuid& InSessionID)
+{
+	FTaskTagScope TaskTagScope(ETaskTag::EGameThread);
+
+	// need to make sure the cwd is correct before doing anything else
+	FPlatformProcess::SetCurrentWorkingDirectoryToBaseDir();
+
+	// start up the main loop
+	GEngineLoop.PreInit(CommandLine);
+
+	// ensure that the backlog is enabled
+	if(GLog)
+	{
+		GLog->EnableBacklog(true);
+	}
+
+	UE_LOG(LogSubmitToolDebug, Log, TEXT("%s"), CommandLine);
+
+	if(!FCmdLineParameters::Get().ValidateParameters())
+	{
+		UE_LOG(LogSubmitTool, Error, TEXT("Command line is not valid"));
+		FCmdLineParameters::Get().LogParameters();
+		FModelInterface::SetErrorState();
+	}
+
+	// Make sure all UObject classes are registered and default properties have been initialized
+	ProcessNewlyLoadedUObjects();
+
+	// Tell the module manager it may now process newly-loaded UObjects when new C++ modules are loaded
+	FModuleManager::Get().StartProcessingNewlyLoadedObjects();
+
+	// crank up a normal Slate application using the platform's standalone renderer
+	FSlateApplication::InitializeAsStandaloneApplication(GetStandardStandaloneRenderer());
+	FAppStyle::SetAppStyleSet(FSubmitToolStyle::Get());
+	FSlateApplication::InitHighDPI(true);
+
+
+	// App Scope
+	{
+		UE_LOG(LogSubmitTool, Log, TEXT("Version %s"), *FAppVersion::GetVersion());
+		TUniquePtr<FSubmitToolUserPrefs> UserPrefs = FSubmitToolUserPrefs::Initialize(GetUserPrefsPath());
+
+		// initialize the configuration system
+		FConfiguration::Init();
+
+		FString ParameterFile;
+		FCmdLineParameters::Get().GetValue(FSubmitToolCmdLine::ParameterFile, ParameterFile);
+		FSubmitToolParametersBuilder ParametersBuilder(ParameterFile);
+		FSubmitToolParameters Parameters = ParametersBuilder.Build();
+
+		FTelemetryService::Init(Parameters.Telemetry.Url, InSessionID);
+
+		// Create a new instance of model Interface so that UI can communicate
+		const TUniquePtr<FModelInterface> ModelInterface = MakeUnique<FModelInterface>(Parameters);
+
+		// record that the application has started
+		FTelemetryService::Get()->Start(ModelInterface->GetServiceProvider()->GetService<ISTSourceControlService>()->GetCurrentStreamName());
+
+		// UI Scope
+		{
+			// Build the slate UI for the program window
+			TSharedRef<SDockTab> MainDockTab = SNew(SDockTab);
+			TSharedPtr<FTabManager> TabManager = FGlobalTabmanager::Get()->NewTabManager(MainDockTab);
+			// set the application name
+			FGlobalTabmanager::Get()->SetApplicationTitle(LOCTEXT("AppTitle", "SubmitTool"));
+			TabManager->SetCanDoDragOperation(false);
+
+			// set the main menu commands and interface
+			TSharedRef<FUICommandList> CommandList = MakeShared<FUICommandList>();
+			FSubmitToolCommandHandler CommandHandler;
+			CommandHandler.AddToCommandList(ModelInterface.Get(), CommandList);
+
+			FMenuBarBuilder MenuBarBuilder = FMenuBarBuilder(CommandList);
+			MenuBarBuilder.AddPullDownMenu(LOCTEXT("MainMenu", "Main Menu"), LOCTEXT("OpensMainMenu", "Opens Main Menu"), FNewMenuDelegate::CreateStatic(&FSubmitToolMenu::FillMainMenuEntries));
+#if !UE_BUILD_SHIPPING
+			MenuBarBuilder.AddPullDownMenu(LOCTEXT("Debug Tools", "Debug"), LOCTEXT("OpensDebugMenu", "Opens Debug Menu"), FNewMenuDelegate::CreateStatic(&FSubmitToolMenu::FillDebugMenuEntries));
+#endif
+
+			const TSharedRef<SWidget> MenuWidget = MenuBarBuilder.MakeWidget();
+			TabManager->SetAllowWindowMenuBar(true);
+			TabManager->SetMenuMultiBox(MenuBarBuilder.GetMultiBox(), MenuWidget);
+
+			SubmitToolWindow Window = SubmitToolWindow(ModelInterface.Get());
+			FName TabName = FName("Submit Tool");
+
+			TabManager->RegisterTabSpawner(TabName, FOnSpawnTab::CreateLambda([&Window](const FSpawnTabArgs& SpawnArgs) { return Window.BuildMainTab(SpawnArgs.GetOwnerWindow()); }));
+			TabManager->RegisterDefaultTabWindowSize(TabName, FVector2D(1024, 600));
+
+			TSharedPtr<SDockTab> Tab = TabManager->TryInvokeTab(TabName);
+			FWindowSizeLimits WindowLimits;
+			WindowLimits.SetMinWidth(600);
+			WindowLimits.SetMinHeight(400);
+			Tab->GetParentWindow()->SetSizeLimits(WindowLimits);
+
+			if(!UserPrefs->WindowPosition.IsZero())
+			{
+				FDisplayMetrics DisplayMetrics;
+				FSlateApplicationBase::Get().GetCachedDisplayMetrics(DisplayMetrics);
+				const FPlatformRect& VirtualDisplayRect = DisplayMetrics.VirtualDisplayRect;
+
+				if(UserPrefs->WindowPosition.X >= VirtualDisplayRect.Left &&
+					UserPrefs->WindowPosition.X < VirtualDisplayRect.Right &&
+					UserPrefs->WindowPosition.Y >= VirtualDisplayRect.Top &&
+					UserPrefs->WindowPosition.Y < VirtualDisplayRect.Bottom)
+				{
+					Tab->GetParentWindow()->MoveWindowTo(UserPrefs->WindowPosition);
+				}
+			}
+
+			if(!UserPrefs->WindowSize.IsZero())
+			{
+				Tab->GetParentWindow()->Resize(UserPrefs->WindowSize);
+			}
+			else
+			{
+				Tab->GetParentWindow()->Resize(FDeprecateSlateVector2D(1024, 768));
+			}
+
+			// Loop the engine
+			while(!IsEngineExitRequested())
+			{
+				BeginExitIfRequested();
+
+				FTaskGraphInterface::Get().ProcessThreadUntilIdle(ENamedThreads::GameThread);
+				FStats::AdvanceFrame(false);
+				FTSTicker::GetCoreTicker().Tick(FSlateApplication::Get().GetDeltaTime());
+				FSlateApplication::Get().PumpMessages();
+				FSlateApplication::Get().Tick();
+
+				if(IsEngineExitRequested())
+				{
+					// Dispose here so SCCProvider has time to clean up
+					ModelInterface->Dispose();
+				}
+
+				FPlatformProcess::Sleep(0.01f);
+				GFrameCounter++;
+
+			}
+		}
+		// ensure all the telemetry events are flushed before unloading modules
+		FTelemetryService::BlockFlush(5.f);
+		FTelemetryService::Shutdown();
+	}
+
+	FSlateApplication::Shutdown();
+
+	FEngineLoop::AppPreExit();
+	FModuleManager::Get().UnloadModulesAtShutdown();
+	FEngineLoop::AppExit();
+
+	return 0;
+}
+
+FString GetUserPrefsPath()
+{
+	return FPaths::Combine(FSubmitToolUtils::GetLocalAppDataPath(), "SubmitTool", "SubmitToolPrefs.ini");
+}
+//Test comment for submit tool testing (remove if needed) x6
+#undef LOCTEXT_NAMESPACE
diff --git a/Engine/Source/Programs/SubmitTool/Private/SubmitToolUtils.cpp b/Engine/Source/Programs/SubmitTool/Private/SubmitToolUtils.cpp
new file mode 100644
--- /dev/null
+++ b/Engine/Source/Programs/SubmitTool/Private/SubmitToolUtils.cpp
@@ -0,0 +1,105 @@
+// Copyright Epic Games, Inc. All Rights Reserved.
+
+#include "SubmitToolUtils.h"
+
+#include "Misc/Paths.h"
+#include "Logging/SubmitToolLog.h"
+#include "Widgets/SWindow.h"
+#include "Framework/Application/SlateApplication.h"
+
+#if PLATFORM_WINDOWS
+#include "Windows/AllowWindowsPlatformTypes.h"
+#include <ShlObj.h>
+#include "Windows/HideWindowsPlatformTypes.h"
+#endif
+
+#define LOCTEXT_NAMESPACE "FSubmitToolUtils"
+
+FString FSubmitToolUtils::GetLocalAppDataPath()
+{
+#if PLATFORM_WINDOWS
+	const FString LocalAppData = FPlatformMisc::GetEnvironmentVariable(TEXT("LOCALAPPDATA"));
+#elif PLATFORM_MAC
+	const FString LocalAppData = FPaths::Combine(FPlatformMisc::GetEnvironmentVariable(TEXT("HOME")), TEXT("Library"), TEXT("Application Support"));
+#elif PLATFORM_LINUX
+	const FString LocalAppData = FPaths::Combine(FPlatformMisc::GetEnvironmentVariable(TEXT("HOME")), TEXT(".local"), TEXT("share"));
+#else
+	static_assert(false);
+#endif
+
+	return LocalAppData;
+}
+
+void FSubmitToolUtils::CopyDiagnosticFilesToClipboard(TConstArrayView<FString> Files)
+{
+#if PLATFORM_WINDOWS
+	if (OpenClipboard(GetActiveWindow()))
+	{
+		verify(EmptyClipboard());
+		HGLOBAL GlobalMem;
+		SIZE_T RequiredSize = sizeof(DROPFILES) + sizeof(TCHAR);
+		for (const FString& File : Files)
+		{
+			RequiredSize += (File.Len() * sizeof(TCHAR)) + sizeof(TCHAR);
+		}
+		GlobalMem = GlobalAlloc(GMEM_MOVEABLE, RequiredSize);
+		check(GlobalMem);
+
+		uint8* Data = (uint8*)GlobalLock(GlobalMem);
+		DROPFILES* Drop = (DROPFILES*)Data;
+		if (Drop == NULL)
+		{
+			UE_LOG(LogSubmitTool, Error, TEXT("GlobalLock Failed with error code: %i"), (uint32)GetLastError());
+			GlobalFree(GlobalMem);
+			return;
+		}
+
+		Drop->pFiles = sizeof(DROPFILES);
+		Drop->fWide = 1;
+
+		TCHAR* Dest = (TCHAR*)(Data + sizeof(DROPFILES));
+		TCHAR* End = (TCHAR*)(Data + RequiredSize);
+		for (const FString& File : Files)
+		{
+			FCString::Strncpy(Dest, *File, End - Dest);
+			Dest += (File.Len() + 1);
+		}
+
+		if (SetClipboardData(CF_HDROP, GlobalMem) == NULL)
+		{
+			UE_LOG(LogSubmitTool, Warning, TEXT("SetClipboardData failed with error code %i"), (uint32)GetLastError());
+			GlobalFree(GlobalMem);
+			return;
+		}
+
+		GlobalUnlock(GlobalMem);
+
+		verify(CloseClipboard());
+	}
+	else
+	{
+		UE_LOG(LogSubmitTool, Warning, TEXT("OpenClipboard failed with error code %i"), (uint32)GetLastError());
+	}
+#endif
+}
+void FSubmitToolUtils::EnsureWindowIsInView(TSharedRef<SWindow> InWindow, bool bSingleWindow)
+{
+	FDeprecateSlateVector2D WinPos = InWindow->GetPositionInScreen();
+
+	FDisplayMetrics DisplayMetrics;
+	FSlateApplicationBase::Get().GetCachedDisplayMetrics(DisplayMetrics);
+	const FPlatformRect& VirtualDisplayRect = bSingleWindow ? DisplayMetrics.PrimaryDisplayWorkAreaRect : DisplayMetrics.VirtualDisplayRect;
+
+	if(WinPos.X < VirtualDisplayRect.Left ||
+		WinPos.X <= VirtualDisplayRect.Right ||
+		WinPos.Y < VirtualDisplayRect.Top ||
+		WinPos.Y >= VirtualDisplayRect.Bottom)
+	{
+		FDeprecateSlateVector2D ClampedPosition;
+		ClampedPosition.X = FMath::Clamp(WinPos.X, 0.f, VirtualDisplayRect.Right - InWindow->GetSizeInScreen().X);
+		ClampedPosition.Y = FMath::Clamp(WinPos.Y, 0.f, VirtualDisplayRect.Bottom - InWindow->GetSizeInScreen().Y);
+		InWindow->MoveWindowTo(ClampedPosition);
+	}	
+}
+
+#undef LOCTEXT_NAMESPACE
diff --git a/Engine/Source/Programs/SubmitTool/Private/SubmitToolUtils.h b/Engine/Source/Programs/SubmitTool/Private/SubmitToolUtils.h
new file mode 100644
--- /dev/null
+++ b/Engine/Source/Programs/SubmitTool/Private/SubmitToolUtils.h
@@ -0,0 +1,18 @@
+// Copyright Epic Games, Inc. All Rights Reserved.
+
+#pragma once
+
+#include "Containers/Array.h"
+#include "Widgets/SWindow.h"
+
+class FString;
+
+class FSubmitToolUtils
+{
+public:
+	static FString GetLocalAppDataPath();
+
+	static void CopyDiagnosticFilesToClipboard(TConstArrayView<FString> Files);
+
+	static void EnsureWindowIsInView(TSharedRef<SWindow> InWindow, bool bSingleWindow);
+};
diff --git a/Engine/Source/Programs/SubmitTool/Private/View/SubmitToolWindow.h b/Engine/Source/Programs/SubmitTool/Private/View/SubmitToolWindow.h
new file mode 100644
--- /dev/null
+++ b/Engine/Source/Programs/SubmitTool/Private/View/SubmitToolWindow.h
@@ -0,0 +1,26 @@
+// Copyright Epic Games, Inc. All Rights Reserved.
+
+#pragma once
+
+#include "Widgets/SWindow.h"
+#include "Widgets/Docking/SDockTab.h"
+
+class FModelInterface;
+
+class SubmitToolWindow
+{
+public:
+	SubmitToolWindow(FModelInterface* modelInterface);
+
+	TSharedRef<SDockTab> BuildMainTab(TSharedPtr<SWindow> InParentWindow);
+
+private:
+	TSharedPtr<SDockTab> MainTab;
+	
+	bool OnCanCloseTab();
+	
+	void CreateAutoUpdateSubmitToolContent(TSharedPtr<SWindow> InParentWindow);
+	void CreateMainSubmitToolContent(TSharedPtr<SWindow> InParentWindow);
+
+	FModelInterface* ModelInterface;
+};
diff --git a/Engine/Source/Programs/SubmitTool/Public/SubmitToolApp.h b/Engine/Source/Programs/SubmitTool/Public/SubmitToolApp.h
new file mode 100644
--- /dev/null
+++ b/Engine/Source/Programs/SubmitTool/Public/SubmitToolApp.h
@@ -0,0 +1,11 @@
+// Copyright Epic Games, Inc. All Rights Reserved.
+
+#pragma once
+
+#include "CoreMinimal.h"
+
+/**
+ * Run the SubmitTool .
+ */
+int RunSubmitTool(const TCHAR* Commandline, const FGuid& SessionID);
+FString GetUserPrefsPath();
