Fix for incorrect UBT Target arguments for programs inside a project, these should be compiled using the Project= argument
Added "dereferencing NULL pointer" to the list of errors to identify in NonUnity validator

diff --git a/Engine/Programs/SubmitTool/Config/SubmitTool.ini b/Engine/Programs/SubmitTool/Config/SubmitTool.ini
--- a/Engine/Programs/SubmitTool/Config/SubmitTool.ini
+++ b/Engine/Programs/SubmitTool/Config/SubmitTool.ini
@@ -135,7 +135,7 @@ Configurations=("Development", "Debug", "Shipping")
 Platforms=("Win64", "Linux", "Mac")
 Targets=("Editor", "Client")
 ToolTip="This Validator compiles the code in the changelist using the configuration/platform provided, it also runs Static Code Analysis at the same time."
-ErrorMessages=("A conflicting instance of")
+ErrorMessages=("A conflicting instance of", "dereferencing NULL pointer")
 UIGroup="Code Validations"
 ExecutionBlockGroups=("UBT")
 TaskArea="LocalFiles"
diff --git a/Engine/Source/Programs/SubmitTool/Private/Logic/Validators/UBTValidator.cpp b/Engine/Source/Programs/SubmitTool/Private/Logic/Validators/UBTValidator.cpp
--- a/Engine/Source/Programs/SubmitTool/Private/Logic/Validators/UBTValidator.cpp
+++ b/Engine/Source/Programs/SubmitTool/Private/Logic/Validators/UBTValidator.cpp
@@ -111,7 +111,6 @@ bool FUBTValidator::Validate(const FString& InCLDescription, const TArray<FSourc
 
 		Arguments = FString::Printf(TEXT("%s %s%s"), *TypedDefinition->ExecutableArguments, *TypedDefinition->StaticAnalyserArg, *StaticAnalyzser);
 	}
-
 	
 
 	// add the file lists for the programs
@@ -121,16 +120,37 @@ bool FUBTValidator::Validate(const FString& InCLDescription, const TArray<FSourc
 		{
 			FString ProgramFileListPath = CreateFileList(ProgramFileList.Value, SubmitToolIntermediateDir + TEXT("SubmitTool/FileLists/"));
 
-			FString TargetStr = FString::Printf(TEXT("%s %s %s -TargetType=%s %s %s\"%s\""),
+			FString TargetStr = FString::Printf(TEXT("%s %s -TargetType=%s %s %s"),
 				*ProgramFileList.Key,
 				*Platform,
-				*Configuration,
 				*Target,
-				*Arguments,
+				*Configuration,
+				*Arguments);
+
+			TSet<FString> ProgramProjectFiles = ExtractUprojectFiles(ProgramFileList.Value);
+			if (ProgramProjectFiles.Num() > 0)
+			{
+				if (ProgramProjectFiles.Num() > 1)
+				{
+					UE_LOG(LogSubmitTool, Warning, TEXT("More than one uprojects found for Program %s, using the first one of:\n%s"), *ProgramFileList.Key, *FString::Join(ProgramProjectFiles, TEXT("\n")));
+				}
+
+				for (const FString& Project : ProgramProjectFiles)
+				{
+					TargetStr = FString::Printf(TEXT("%s %s\"%s\""),
+						*TargetStr,
+						*TypedDefinition->ProjectArgument,
+						*Project);
+
+					break;
+				}
+			}
+
+			TargetStr = FString::Printf(TEXT("%s %s\"%s\""),
+				*TargetStr,
 				*TypedDefinition->FileListArgument,
 				*ProgramFileListPath);
 
-
 			UE_LOG(LogValidators, Log, TEXT("[%s] Using Target: %s"), *GetValidatorName(), *TargetStr);
 			Targets.Add(TargetStr);
 		}
@@ -150,11 +170,11 @@ bool FUBTValidator::Validate(const FString& InCLDescription, const TArray<FSourc
 		{
 
 			FString ProjectFileListPath = CreateFileList(ProjectFilteredFiles, SubmitToolIntermediateDir + TEXT("SubmitTool/FileLists/"));
-			FString TargetStr = FString::Printf(TEXT("%s %s %s -TargetType=%s %s\"%s\" %s\"%s\""),
+			FString TargetStr = FString::Printf(TEXT("%s %s -TargetType=%s %s %s\"%s\" %s\"%s\""),
 				*Platform,
 				*Configuration,
-				*Arguments,
 				*Target,
+				*Arguments,
 				*TypedDefinition->ProjectArgument,
 				*ProjectFile,
 				*TypedDefinition->FileListArgument,
