Fixed cancellation of TagValidator causing issues and invalid states
Added message to integration request being blocked when validations are not finished.
Ensure Integration window is in the right position every time we push the open button
Submit on success can no longer be focused
Allow tabbing in Description textbox
Fixed issue with incorrect preflight configuration

diff --git a/Engine/Source/Programs/SubmitTool/Private/Logic/PreflightService.cpp b/Engine/Source/Programs/SubmitTool/Private/Logic/PreflightService.cpp
--- a/Engine/Source/Programs/SubmitTool/Private/Logic/PreflightService.cpp
+++ b/Engine/Source/Programs/SubmitTool/Private/Logic/PreflightService.cpp
@@ -621,7 +621,8 @@ void FPreflightService::OnCreateShelveOperationComplete(const FSourceControlOper
 
 bool FPreflightService::SelectPreflightTemplate(FPreflightTemplateDefinition& OutTemplate) const
 {
-	bool bFoundTemplate = false;
+	FConfiguration::AddOrUpdateEntry(TEXT("$(StreamRoot)"),	ServiceProvider->GetService<ISTSourceControlService>()->GetRootStreamName());
+
 	TSharedPtr<FChangelistService> ChangelistService = ServiceProvider->GetService<FChangelistService>();
 
 	// Get the list of files in the changelist
diff --git a/Engine/Source/Programs/SubmitTool/Private/Logic/TasksService.cpp b/Engine/Source/Programs/SubmitTool/Private/Logic/TasksService.cpp
--- a/Engine/Source/Programs/SubmitTool/Private/Logic/TasksService.cpp
+++ b/Engine/Source/Programs/SubmitTool/Private/Logic/TasksService.cpp
@@ -459,3 +459,14 @@ void FTasksService::StopTasks(const FName& InTaskId)
 		}
 	}
 }
+
+void FTasksService::StopTasksByArea(const ETaskArea& InArea)
+{
+	for (const TPair<FName, TSharedPtr<FValidatorBase>>& TaskPair : Tasks)
+	{
+		if ((TaskPair.Value->Definition->TaskArea & InArea) != ETaskArea::None)
+		{
+			TaskPair.Value->CancelValidation();
+		}
+	}
+}
diff --git a/Engine/Source/Programs/SubmitTool/Private/Logic/TasksService.h b/Engine/Source/Programs/SubmitTool/Private/Logic/TasksService.h
--- a/Engine/Source/Programs/SubmitTool/Private/Logic/TasksService.h
+++ b/Engine/Source/Programs/SubmitTool/Private/Logic/TasksService.h
@@ -45,6 +45,7 @@ public:
 	virtual const TArray<FString> GetAddendums() const;
 
 	virtual void StopTasks(const FName& InValidatorId = FName());
+	virtual void StopTasksByArea(const ETaskArea& InArea);
 
 	FOnTaskFinished OnTasksQueueFinished;
 	FOnSingleTaskFinished OnSingleTaskFinished;
diff --git a/Engine/Source/Programs/SubmitTool/Private/Logic/Validators/PreflightValidator.cpp b/Engine/Source/Programs/SubmitTool/Private/Logic/Validators/PreflightValidator.cpp
--- a/Engine/Source/Programs/SubmitTool/Private/Logic/Validators/PreflightValidator.cpp
+++ b/Engine/Source/Programs/SubmitTool/Private/Logic/Validators/PreflightValidator.cpp
@@ -176,6 +176,12 @@ void FPreflightValidator::ValidatePreflights(const TUniquePtr<FPreflightList>& I
 		return;
 	}
 
+	if (RunningPreflights > 0)
+	{
+		UE_LOG(LogValidators, Log, TEXT("[%s] Periodically checking updated horde state... If you see your PF has finished in horde you can force a refresh in the preflight tag refresh button"), *GetValidatorName());
+		UE_LOG(LogValidatorsResult, Log, TEXT("[%s] Periodically checking updated horde state... If you see your PF has finished in horde you can force a refresh in the preflight tag refresh button"), *GetValidatorName());
+	}
+
 	// Fail early even if other PF are running
 	if(!FailedPreflights.IsEmpty())
 	{
diff --git a/Engine/Source/Programs/SubmitTool/Private/Logic/Validators/TagValidator.cpp b/Engine/Source/Programs/SubmitTool/Private/Logic/Validators/TagValidator.cpp
--- a/Engine/Source/Programs/SubmitTool/Private/Logic/Validators/TagValidator.cpp
+++ b/Engine/Source/Programs/SubmitTool/Private/Logic/Validators/TagValidator.cpp
@@ -19,8 +19,9 @@ bool FTagValidator::Validate(const FString& InCLDescription, const TArray<FSourc
 	}
 
 	// Wait for the GetUsersAndGroups asynchronously to finish validation
-	ServiceProvider.Pin()->GetService<ISTSourceControlService>()->GetUsersAndGroups(FOnUsersAndGroupsGet::FDelegate::CreateLambda([&bPasses, this, InTags](TArray<TSharedPtr<FUserData>>& InP4Users, TArray<TSharedPtr<FString>>& InP4Groups)
+	ServiceProvider.Pin()->GetService<ISTSourceControlService>()->GetUsersAndGroups(FOnUsersAndGroupsGet::FDelegate::CreateLambda([bPasses, this, InTags](TArray<TSharedPtr<FUserData>>& InP4Users, TArray<TSharedPtr<FString>>& InP4Groups)
 		{
+			bool bResult = bPasses;
 			for (const FTag* tag : InTags)
 			{
 				if (tag->IsEnabled() && tag->Definition.bIsUserValue)
@@ -35,11 +36,11 @@ bool FTagValidator::Validate(const FString& InCLDescription, const TArray<FSourc
 				else
 				{
 					tag->SetTagState(ETagState::Failed);
-					bPasses = false;
+					bResult = false;
 				}
 			}
 
-			ValidationFinished(bPasses);
+			ValidationFinished(bResult);
 		}));
 
 	return true;
diff --git a/Engine/Source/Programs/SubmitTool/Private/Models/ModelInterface.h b/Engine/Source/Programs/SubmitTool/Private/Models/ModelInterface.h
--- a/Engine/Source/Programs/SubmitTool/Private/Models/ModelInterface.h
+++ b/Engine/Source/Programs/SubmitTool/Private/Models/ModelInterface.h
@@ -127,7 +127,7 @@ public:
 
 	void ValidateChangelist() const													{ ValidationService->QueueAll(); }
 	void ValidateSingle(const FName& ValidatorId) const								{ ValidationService->QueueSingle(ValidatorId, true); }
-	void ValidateCLDescription() const												{ ValidationService->QueueByArea(ETaskArea::Changelist); }
+	void ValidateCLDescription() const												{ ValidationService->StopTasksByArea(ETaskArea::Changelist); ValidationService->QueueByArea(ETaskArea::Changelist); }
 	bool IsCLValid() const															{ return ValidationService->GetIsRunSuccessful(); }
 	bool CanLaunchPreflight() const;
 	void ReevaluateSubmitToolTag();
diff --git a/Engine/Source/Programs/SubmitTool/Private/View/Widgets/SIntegrationWidget.cpp b/Engine/Source/Programs/SubmitTool/Private/View/Widgets/SIntegrationWidget.cpp
--- a/Engine/Source/Programs/SubmitTool/Private/View/Widgets/SIntegrationWidget.cpp
+++ b/Engine/Source/Programs/SubmitTool/Private/View/Widgets/SIntegrationWidget.cpp
@@ -54,7 +54,7 @@ void SIntegrationWidget::Construct(const FArguments& InArgs)
 		.Padding(5, 5)
 		.AutoSize()
 		[
-				SNew(SHorizontalBox)
+			SNew(SHorizontalBox)
 			+ SHorizontalBox::Slot()
 			.Padding(FAppStyle::GetMargin("StandardDialog.SlotPadding"))
 			.AutoWidth()
@@ -70,6 +70,21 @@ void SIntegrationWidget::Construct(const FArguments& InArgs)
 		+SScrollBox::Slot()
 		.Padding(5, 5)
 		.AutoSize()
+		[
+			SNew(SHorizontalBox)
+			+ SHorizontalBox::Slot()
+			.Padding(FAppStyle::GetMargin("StandardDialog.SlotPadding"))
+			.AutoWidth()
+			[
+				SNew(STextBlock)
+				.ColorAndOpacity(FStyleColors::Error)
+				.Visibility_Lambda([&ModelInterface = ModelInterface]{ return ModelInterface->HasSubmitToolTag() ? EVisibility::Collapsed : EVisibility::All; })
+				.Text(FText::FromString(TEXT("Required validations have failed or need to finish running.")))
+			]
+		]
+		+SScrollBox::Slot()
+		.Padding(5, 5)
+		.AutoSize()
 		[
 			SNew(SHorizontalBox)
 			+ SHorizontalBox::Slot()
@@ -119,19 +134,20 @@ void SIntegrationWidget::Open()
 		ParentWindow->SetRequestDestroyWindowOverride(FRequestDestroyWindowOverride::CreateLambda([](const TSharedRef<SWindow>& Window) { Window->HideWindow(); }));
 		FSlateApplication::Get().AddWindowAsNativeChild(ParentWindow.ToSharedRef(), MainWindow.ToSharedRef(), false);
 
-		FDeprecateSlateVector2D NewPosition = MainWindow->GetPositionInScreen();
-		NewPosition.X += MainWindow->GetSizeInScreen().X;
-		ParentWindow->MoveWindowTo(NewPosition);
-
-		float ySize = ParentWindow->IsWindowMaximized() ? 300 : MainWindow->GetSizeInScreen().Y - 40;
-		ParentWindow->Resize(FDeprecateSlateVector2D(400, ySize));
-
-		FSubmitToolUtils::EnsureWindowIsInView(ParentWindow.ToSharedRef(), true);
-
 		ParentWindow->SetContent(AsShared());
 	}
 
 	ParentWindow->SetTitle(FText::FromString(FString::Printf(TEXT("Integration information for CL %s"), *ModelInterface->GetCLID())));
+
+	FDeprecateSlateVector2D NewPosition = MainWindow->GetPositionInScreen();
+	NewPosition.X += MainWindow->GetSizeInScreen().X;
+	ParentWindow->MoveWindowTo(NewPosition);
+
+	float ySize = ParentWindow->IsWindowMaximized() ? 300 : MainWindow->GetSizeInScreen().Y - 40;
+	ParentWindow->Resize(FDeprecateSlateVector2D(400, ySize));
+
+	FSubmitToolUtils::EnsureWindowIsInView(ParentWindow.ToSharedRef(), true);
+
 	ParentWindow->BringToFront();
 	ParentWindow->ShowWindow();
 
diff --git a/Engine/Source/Programs/SubmitTool/Private/View/Widgets/SubmitToolWidget.cpp b/Engine/Source/Programs/SubmitTool/Private/View/Widgets/SubmitToolWidget.cpp
--- a/Engine/Source/Programs/SubmitTool/Private/View/Widgets/SubmitToolWidget.cpp
+++ b/Engine/Source/Programs/SubmitTool/Private/View/Widgets/SubmitToolWidget.cpp
@@ -157,6 +157,7 @@ void SubmitToolWidget::Construct(const FArguments& InArgs)
 		.OnTextChanged_Lambda([&ModelInterface = ModelInterface](const FText& newText) { ModelInterface->SetCLDescription(newText); })
 		.OnTextCommitted_Lambda([this](const FText& Text, ETextCommit::Type CommitType) { if(CommitType != ETextCommit::OnEnter) { ModelInterface->ValidateCLDescription(); }})
 		.AutoWrapText(true)
+		.OnIsTypedCharValid(FOnIsTypedCharValid::CreateLambda([](const TCHAR) { return true; }))
 		.IsReadOnly(false)
 		.IsEnabled_Static(&FModelInterface::GetInputEnabled);
 
@@ -402,6 +403,7 @@ TSharedRef<SHorizontalBox> SubmitToolWidget::BuildButtonRow()
 				SNew(SCheckBox)
 				.IsChecked_Lambda([this]() { return ModelInterface->bSubmitOnSuccessfulValidation ? ECheckBoxState::Checked : ECheckBoxState::Unchecked; })
 				.OnCheckStateChanged_Lambda([this](ECheckBoxState InNewState) { ModelInterface->bSubmitOnSuccessfulValidation = !ModelInterface->bSubmitOnSuccessfulValidation; })
+				.IsFocusable(false)
 			]
 			+SHorizontalBox::Slot()
 			.AutoWidth()
