- Fixed submit tool growing is size by tiny amounts every time it was opened
- Fixed lifetime issues due to circular dependencies, ensuring tasks are killed when ST closes
- Added safeguards for overwritting a service registration
- Added capability of registering a service type more than once with different ids
- Fix for race condition when getting p4 users and groups

diff --git a/Engine/Source/Programs/SubmitTool/Private/Logic/JiraService.cpp b/Engine/Source/Programs/SubmitTool/Private/Logic/JiraService.cpp
--- a/Engine/Source/Programs/SubmitTool/Private/Logic/JiraService.cpp
+++ b/Engine/Source/Programs/SubmitTool/Private/Logic/JiraService.cpp
@@ -20,7 +20,7 @@
 #include "Configuration/Configuration.h"
 
 
-FJiraService::FJiraService(const FJiraParameters& InJiraSettings, const int32 InMaxResults, TSharedRef<FSubmitToolServiceProvider> InServiceProvider) :
+FJiraService::FJiraService(const FJiraParameters& InJiraSettings, const int32 InMaxResults, TWeakPtr<FSubmitToolServiceProvider> InServiceProvider) :
 	Definition(InJiraSettings),
 	MaxResults(InMaxResults),
 	ServiceProvider(InServiceProvider)
@@ -43,7 +43,7 @@ FJiraService::~FJiraService()
 
 bool FJiraService::FetchJiraTickets(bool InForce)
 {
-	TSharedPtr<FCredentialsService> CredentialsService = ServiceProvider->GetService<FCredentialsService>();
+	TSharedPtr<FCredentialsService> CredentialsService = ServiceProvider.Pin()->GetService<FCredentialsService>();
 	if(CredentialsService->HasCredentials() && !Definition.ServerAddress.IsEmpty())
 	{
 		if (InForce || (JiraIssues.Num() == 0 && CredentialsService->AreCredentialsValid()))
@@ -73,7 +73,7 @@ void FJiraService::QueryIssues()
 	
 	JiraRequest->OnProcessRequestComplete().BindRaw(this, &FJiraService::QueryIssues_HttpRequestComplete);
 
-	TSharedPtr<FCredentialsService> CredentialsService = ServiceProvider->GetService<FCredentialsService>();
+	TSharedPtr<FCredentialsService> CredentialsService = ServiceProvider.Pin()->GetService<FCredentialsService>();
 	FString Url = FString::Format(TEXT("https://{0}/rest/api/2/search?maxResults={1}&jql=assignee={2}"), { Definition.ServerAddress, this->MaxResults, CredentialsService->GetUsername()});
 
 	JiraRequest->SetURL(Url);
@@ -88,7 +88,7 @@ void FJiraService::QueryIssues_HttpRequestComplete(FHttpRequestPtr HttpRequest,
 {
 	JiraRequest = nullptr;
 	bOngoingRequest = false;
-	TSharedPtr<FCredentialsService> CredentialsService = ServiceProvider->GetService<FCredentialsService>();
+	TSharedPtr<FCredentialsService> CredentialsService = ServiceProvider.Pin()->GetService<FCredentialsService>();
 
 	if (!bSucceeded)
 	{
@@ -247,7 +247,7 @@ void FJiraService::SaveJiraIssues() const
 void FJiraService::LoadJiraIssues()
 {
 	// do not load the issues if there is no credentials
-	if (!ServiceProvider->GetService<FCredentialsService>()->HasCredentials())
+	if (!ServiceProvider.Pin()->GetService<FCredentialsService>()->HasCredentials())
 	{
 		return;
 	}
@@ -318,7 +318,7 @@ void FJiraService::GetuserInfo()
 
 	HttpRequest->SetURL(Url);
 
-	HttpRequest->SetHeader(TEXT("Authorization"), TEXT("Basic ") + ServiceProvider->GetService<FCredentialsService>()->GetEncodedLoginString());
+	HttpRequest->SetHeader(TEXT("Authorization"), TEXT("Basic ") + ServiceProvider.Pin()->GetService<FCredentialsService>()->GetEncodedLoginString());
 	HttpRequest->SetVerb(TEXT("GET"));
 	HttpRequest->ProcessRequest();
 }
@@ -335,7 +335,7 @@ void FJiraService::GetuserInfo_HttpRequestComplete(FHttpRequestPtr HttpRequest,
 	{
 		if (EHttpResponseCodes::IsOk(HttpResponse->GetResponseCode()))
 		{
-			ServiceProvider->GetService<FCredentialsService>()->SetCredentialsValid(true);
+			ServiceProvider.Pin()->GetService<FCredentialsService>()->SetCredentialsValid(true);
 			FString ResponseStr = HttpResponse->GetContentAsString();
 
 			TSharedPtr<FJsonObject> RootJsonObject;
@@ -366,7 +366,7 @@ void FJiraService::GetIssueAndCreateServiceDeskRequest(const FString& Key, const
 	
 	FString Url = FString::Format(TEXT("https://{0}/rest/api/2/issue/{1}"), { Definition.ServerAddress, Key });
 	HttpRequest->SetURL(Url);
-	HttpRequest->SetHeader(TEXT("Authorization"), TEXT("Basic ") + ServiceProvider->GetService<FCredentialsService>()->GetEncodedLoginString());
+	HttpRequest->SetHeader(TEXT("Authorization"), TEXT("Basic ") + ServiceProvider.Pin()->GetService<FCredentialsService>()->GetEncodedLoginString());
 	HttpRequest->SetVerb(TEXT("GET"));
 
 	HttpRequest->ProcessRequest();
@@ -460,7 +460,7 @@ void FJiraService::CreateServiceDeskRequest(TSharedPtr<FJsonObject> InBaseJiraJs
 
 	if(!Definition.PreflightField.IsEmpty())
 	{		
-		const FTag* PreflightTag = ServiceProvider->GetService<FTagService>()->GetTag(TEXT("#preflight"));
+		const FTag* PreflightTag = ServiceProvider.Pin()->GetService<FTagService>()->GetTag(TEXT("#preflight"));
 		if (PreflightTag != nullptr && PreflightTag->GetValues().Num() > 0)
 		{
 			FString PreflightTagValue = PreflightTag->GetValues()[0];
@@ -469,7 +469,7 @@ void FJiraService::CreateServiceDeskRequest(TSharedPtr<FJsonObject> InBaseJiraJs
 			{
 				if(!PreflightTagValue.Contains(TEXT("/job/")))
 				{
-					RequestFieldValuesJson->SetStringField(Definition.PreflightField, FString::Format(TEXT("{0}job/{1}"), { ServiceProvider->GetService<FPreflightService>()->GetHordeServerAddress(), PreflightTagValue }));
+					RequestFieldValuesJson->SetStringField(Definition.PreflightField, FString::Format(TEXT("{0}job/{1}"), { ServiceProvider.Pin()->GetService<FPreflightService>()->GetHordeServerAddress(), PreflightTagValue }));
 				}
 				else
 				{
@@ -480,7 +480,7 @@ void FJiraService::CreateServiceDeskRequest(TSharedPtr<FJsonObject> InBaseJiraJs
 	}
 
 	FString Requestor = FConfiguration::Substitute(TEXT("$(USER)"));
-	TSharedPtr<FUserData> LocalUserData = ServiceProvider->GetService<ISTSourceControlService>()->GetUserDataFromCache(Requestor);
+	TSharedPtr<FUserData> LocalUserData = ServiceProvider.Pin()->GetService<ISTSourceControlService>()->GetUserDataFromCache(Requestor);
 	if (LocalUserData.IsValid())
 	{
 		Requestor = LocalUserData.Get()->Email;
diff --git a/Engine/Source/Programs/SubmitTool/Private/Logic/JiraService.h b/Engine/Source/Programs/SubmitTool/Private/Logic/JiraService.h
--- a/Engine/Source/Programs/SubmitTool/Private/Logic/JiraService.h
+++ b/Engine/Source/Programs/SubmitTool/Private/Logic/JiraService.h
@@ -20,7 +20,7 @@ class FJiraService final : public ISubmitToolService
 {
 public:
 	FJiraService() = delete;
-	FJiraService(const FJiraParameters& InJiraSettings, const int32 InMaxResults, TSharedRef<FSubmitToolServiceProvider> InServiceProvider);
+	FJiraService(const FJiraParameters& InJiraSettings, const int32 InMaxResults, TWeakPtr<FSubmitToolServiceProvider> InServiceProvider);
 	~FJiraService();
 
 	bool FetchJiraTickets(bool InForce = false);
@@ -63,7 +63,7 @@ private:
 	int32 TotalIssues;
 	TMap<FString, FJiraIssue> JiraIssues;
 
-	TSharedRef<FSubmitToolServiceProvider> ServiceProvider;
+	TWeakPtr<FSubmitToolServiceProvider> ServiceProvider;
 };
 
 Expose_TNameOf(FJiraService);
\ No newline at end of file
diff --git a/Engine/Source/Programs/SubmitTool/Private/Logic/P4LockdownService.cpp b/Engine/Source/Programs/SubmitTool/Private/Logic/P4LockdownService.cpp
--- a/Engine/Source/Programs/SubmitTool/Private/Logic/P4LockdownService.cpp
+++ b/Engine/Source/Programs/SubmitTool/Private/Logic/P4LockdownService.cpp
@@ -12,7 +12,7 @@
 #include "Internationalization/Regex.h"
 #include "Configuration/Configuration.h"
 
-FP4LockdownService::FP4LockdownService(const FP4LockdownParameters& InParameters, TSharedPtr<FSubmitToolServiceProvider> InServiceProvider)
+FP4LockdownService::FP4LockdownService(const FP4LockdownParameters& InParameters, TWeakPtr<FSubmitToolServiceProvider> InServiceProvider)
 : 	ServiceProvider(InServiceProvider), Parameters(InParameters)
 {
 	using namespace UE::Tasks;
@@ -172,7 +172,7 @@ FP4LockdownService::FPathLockdownResult FP4LockdownService::IsPathInLockdown(con
 bool FP4LockdownService::GetConfigFile(const FString& InConfigId, const FString& InDepotPath)
 {
 	TArray<FSharedBuffer> FileBuffers;
-	if(ServiceProvider->GetService<ISTSourceControlService>()->DownloadFiles(InDepotPath, FileBuffers).GetResult() && FileBuffers.Num() > 0)
+	if(ServiceProvider.Pin()->GetService<ISTSourceControlService>()->DownloadFiles(InDepotPath, FileBuffers).GetResult() && FileBuffers.Num() > 0)
 	{
 		FScopeLock Scope(&FileMutex);
 		
diff --git a/Engine/Source/Programs/SubmitTool/Private/Logic/P4LockdownService.h b/Engine/Source/Programs/SubmitTool/Private/Logic/P4LockdownService.h
--- a/Engine/Source/Programs/SubmitTool/Private/Logic/P4LockdownService.h
+++ b/Engine/Source/Programs/SubmitTool/Private/Logic/P4LockdownService.h
@@ -17,7 +17,7 @@ class FP4LockdownService : public ISubmitToolService
 {	
 public:
 	FP4LockdownService() = delete;
-	FP4LockdownService(const FP4LockdownParameters& InParameters, TSharedPtr<FSubmitToolServiceProvider> InServiceProvider);
+	FP4LockdownService(const FP4LockdownParameters& InParameters, TWeakPtr<FSubmitToolServiceProvider> InServiceProvider);
 	bool ArePathsInLockdown(const TArray<FString>& InPaths, bool& bOutAllowlisted);
 	bool IsBlockingOperationRunning() const
 	{
@@ -52,7 +52,7 @@ private:
 
 	TMap<FString, FSharedBuffer> DownloadedFiles;
 	UE::Tasks::TTask<bool> ConfigFilesTask;
-	TSharedPtr<FSubmitToolServiceProvider> ServiceProvider;
+	TWeakPtr<FSubmitToolServiceProvider> ServiceProvider;
 	const FP4LockdownParameters& Parameters;
 	TArray<FAllowListData> AllowListData;
 	TArray<FOverrideData> OverrideData;
diff --git a/Engine/Source/Programs/SubmitTool/Private/Logic/PreflightService.cpp b/Engine/Source/Programs/SubmitTool/Private/Logic/PreflightService.cpp
--- a/Engine/Source/Programs/SubmitTool/Private/Logic/PreflightService.cpp
+++ b/Engine/Source/Programs/SubmitTool/Private/Logic/PreflightService.cpp
@@ -23,7 +23,7 @@
 FPreflightService::FPreflightService(
 	const FHordeParameters& InSettings,
 	FModelInterface* InModelInterface,
-	TSharedRef<FSubmitToolServiceProvider> InServiceProvider)
+	TWeakPtr<FSubmitToolServiceProvider> InServiceProvider)
 	:
 	Definition(InSettings),
 	ServiceProvider(InServiceProvider),
@@ -36,7 +36,7 @@ FPreflightService::FPreflightService(
 	DeleteShelveCallback = FSourceControlOperationComplete::CreateRaw(this, &FPreflightService::OnDeleteShelveOperationComplete);
 	CreateShelveCallback = FSourceControlOperationComplete::CreateRaw(this, &FPreflightService::OnCreateShelveOperationComplete);
 
-	PreflightTag = ServiceProvider->GetService<FTagService>()->GetTag(TEXT("#preflight"));
+	PreflightTag = ServiceProvider.Pin()->GetService<FTagService>()->GetTag(TEXT("#preflight"));
 }
 
 FPreflightService::~FPreflightService()
@@ -47,7 +47,7 @@ FPreflightService::~FPreflightService()
 
 bool FPreflightService::Tick(float DeltaTime)
 {
-	TSharedPtr<FChangelistService> ChangelistService = ServiceProvider->GetService<FChangelistService>();
+	TSharedPtr<FChangelistService> ChangelistService = ServiceProvider.Pin()->GetService<FChangelistService>();
 	switch(State)
 	{
 		//////////////////////////////////////////////////
@@ -128,7 +128,7 @@ TMap<FString, FStringFormatArg> FPreflightService::GetFormatParameters() const
 	TMap<FString, FStringFormatArg> FormatMap =
 	{
 		{ TEXT("URL"), Definition.HordeServerAddress },
-		{ TEXT("CLID"), ServiceProvider->GetService<FChangelistService>()->GetCLID()}
+		{ TEXT("CLID"), ServiceProvider.Pin()->GetService<FChangelistService>()->GetCLID()}
 	};
 
 	FormatMap.Add(TEXT("Stream"), StreamName);
@@ -153,7 +153,7 @@ void FPreflightService::RequestPreflight()
 {
 	if(State == EPreflightServiceState::Idle)
 	{
-		if(!ServiceProvider->GetService<ISTSourceControlService>()->GetClientStreams().IsEmpty())
+		if(!ServiceProvider.Pin()->GetService<ISTSourceControlService>()->GetClientStreams().IsEmpty())
 		{
 			UE_LOG(LogSubmitTool, Log, TEXT("Preflight: Requesting..."));
 			State = EPreflightServiceState::RequestDeleteShelve;
@@ -204,13 +204,13 @@ void FPreflightService::Requeue()
 
 void FPreflightService::FetchPreflightInfo(bool bRequeue, const FString& InOAuthToken)
 {
-	TSharedPtr<FCredentialsService> CredentialsService = ServiceProvider->GetService<FCredentialsService>();
+	TSharedPtr<FCredentialsService> CredentialsService = ServiceProvider.Pin()->GetService<FCredentialsService>();
 	if(Definition.HordeServerAddress.IsEmpty() || !CredentialsService->IsOIDCTokenEnabled())
 	{
 		return;
 	}
 
-	TSharedPtr<FChangelistService> ChangelistService = ServiceProvider->GetService<FChangelistService>();
+	TSharedPtr<FChangelistService> ChangelistService = ServiceProvider.Pin()->GetService<FChangelistService>();
 
 	// Don't bother with the default changelist
 	if(ChangelistService->GetCLID() == TEXT("default"))
@@ -313,7 +313,7 @@ void FPreflightService::FetchPreflightInfo(bool bRequeue, const FString& InOAuth
 							// Only log when there's a different number of preflights
 							if(!HordePreflights.IsValid() || HordePreflights->PreflightList.Num() != NewHordePreflights->PreflightList.Num())
 							{
-								UE_LOG(LogSubmitTool, Log, TEXT("Retrieved %d preflights for CL %s"), NewHordePreflights->PreflightList.Num(), *ServiceProvider->GetService<FChangelistService>()->GetCLID());
+								UE_LOG(LogSubmitTool, Log, TEXT("Retrieved %d preflights for CL %s"), NewHordePreflights->PreflightList.Num(), *ServiceProvider.Pin()->GetService<FChangelistService>()->GetCLID());
 							}
 
 							HordePreflights = MoveTemp(NewHordePreflights);
@@ -509,10 +509,10 @@ void FPreflightService::FetchUnlinkedPreflight(const FString& InPreflightId, boo
 
 void FPreflightService::StartPreflight()
 {
-	TSharedPtr<ISTSourceControlService> SCCService = ServiceProvider->GetService<ISTSourceControlService>();
+	TSharedPtr<ISTSourceControlService> SCCService = ServiceProvider.Pin()->GetService<ISTSourceControlService>();
 	StreamName = SCCService->GetRootStreamName();
 
-	const TArray<FString>& ShelvedPaths = ServiceProvider->GetService<FChangelistService>()->GetShelvedFilesDepotPaths(true);
+	const TArray<FString>& ShelvedPaths = ServiceProvider.Pin()->GetService<FChangelistService>()->GetShelvedFilesDepotPaths(true);
 
 	if (ShelvedPaths.IsEmpty())
 	{
@@ -621,9 +621,9 @@ void FPreflightService::OnCreateShelveOperationComplete(const FSourceControlOper
 
 bool FPreflightService::SelectPreflightTemplate(FPreflightTemplateDefinition& OutTemplate) const
 {
-	FConfiguration::AddOrUpdateEntry(TEXT("$(StreamRoot)"),	ServiceProvider->GetService<ISTSourceControlService>()->GetRootStreamName());
+	FConfiguration::AddOrUpdateEntry(TEXT("$(StreamRoot)"),	ServiceProvider.Pin()->GetService<ISTSourceControlService>()->GetRootStreamName());
 
-	TSharedPtr<FChangelistService> ChangelistService = ServiceProvider->GetService<FChangelistService>();
+	TSharedPtr<FChangelistService> ChangelistService = ServiceProvider.Pin()->GetService<FChangelistService>();
 
 	// Get the list of files in the changelist
 	const TArray<FString>& FilesInCl = bCheckShelveInstead ? ChangelistService->GetShelvedFilesDepotPaths() : ChangelistService->GetFilesDepotPaths();
@@ -652,7 +652,7 @@ FString FPreflightService::GetAdditionalTasksString(const FPreflightTemplateDefi
 	TStringBuilder<256> AdditionalTaskStrBuilder;
 	const FString BaseString = TEXT("&id-additional-tasks.");
 	const FString EndString = TEXT("=true");
-	TSharedPtr<FChangelistService> ChangelistService = ServiceProvider->GetService<FChangelistService>();
+	TSharedPtr<FChangelistService> ChangelistService = ServiceProvider.Pin()->GetService<FChangelistService>();
 
 	const TArray<FString>& FilesInCl = bCheckShelveInstead ? ChangelistService->GetShelvedFilesDepotPaths() : ChangelistService->GetFilesDepotPaths();
 	for(const FPreflightAdditionalTask& AdditionalTask : InTemplate.AdditionalTasks)
@@ -681,7 +681,7 @@ EDialogFactoryResult FPreflightService::ShowRecreateShelveDialog() const
 		ModelInterface->GetMainTab().Pin()->GetParentWindow()->DrawAttention(FWindowDrawAttentionParameters());
 	}
 
-	const FText TextTitle = FText::FromString(FString::Printf(TEXT("Preflight CL %s: Recreate shelved files?"), *ServiceProvider->GetService<FChangelistService>()->GetCLID()));
+	const FText TextTitle = FText::FromString(FString::Printf(TEXT("Preflight CL %s: Recreate shelved files?"), *ServiceProvider.Pin()->GetService<FChangelistService>()->GetCLID()));
 	const FText TextDescription = FText::FromString(TEXT("There are already shelved files in this changelist.\n\nDo you want to delete and recreate your shelf from the latest changes in your local files for use in the preflight?"));
 
 	return FDialogFactory::ShowDialog(TextTitle, TextDescription, TArray<FString>{ TEXT("Re-shelve files"), TEXT("Use existing Shelve"), TEXT("Cancel") });
@@ -694,7 +694,7 @@ EDialogFactoryResult FPreflightService::ShowUpdatePreflightTagDialog() const
 		ModelInterface->GetMainTab().Pin()->GetParentWindow()->DrawAttention(FWindowDrawAttentionParameters());
 	}
 
-	const FText TextTitle = FText::FromString(FString::Printf(TEXT("Preflight CL %s: Newer preflight available"), *ServiceProvider->GetService<FChangelistService>()->GetCLID()));
+	const FText TextTitle = FText::FromString(FString::Printf(TEXT("Preflight CL %s: Newer preflight available"), *ServiceProvider.Pin()->GetService<FChangelistService>()->GetCLID()));
 	const FText TextDescription = FText::FromString(FString::Printf(TEXT("There is a newer preflight for this changelist:\n<a id=\"browser\" style=\"Hyperlink\" href=\"%sjob/%s\">%s - %s</>\n\nDo you want to update the #preflight tag?"), *Definition.HordeServerAddress, *HordePreflights->PreflightList[0].ID, *HordePreflights->PreflightList[0].Name, *HordePreflights->PreflightList[0].ID));
 
 	return FDialogFactory::ShowDialog(TextTitle, TextDescription, TArray<FString>{ TEXT("Update Tag Value"), TEXT("Cancel") });
diff --git a/Engine/Source/Programs/SubmitTool/Private/Logic/PreflightService.h b/Engine/Source/Programs/SubmitTool/Private/Logic/PreflightService.h
--- a/Engine/Source/Programs/SubmitTool/Private/Logic/PreflightService.h
+++ b/Engine/Source/Programs/SubmitTool/Private/Logic/PreflightService.h
@@ -34,7 +34,7 @@ public:
 	FPreflightService(
 		const FHordeParameters& InSettings,
 		FModelInterface* InModelInterface,
-		TSharedRef<FSubmitToolServiceProvider> InServiceProvider);
+		TWeakPtr<FSubmitToolServiceProvider> InServiceProvider);
 	~FPreflightService();
 
 	void RequestPreflight();
@@ -86,7 +86,7 @@ private:
 	const FHordeParameters Definition;
 
 	// services we depend on
-	TSharedRef<FSubmitToolServiceProvider> ServiceProvider;
+	TWeakPtr<FSubmitToolServiceProvider> ServiceProvider;
 	FTag* PreflightTag;
 	FModelInterface* ModelInterface;
 
diff --git a/Engine/Source/Programs/SubmitTool/Private/Logic/Services/SourceControl/SubmitToolPerforce.cpp b/Engine/Source/Programs/SubmitTool/Private/Logic/Services/SourceControl/SubmitToolPerforce.cpp
--- a/Engine/Source/Programs/SubmitTool/Private/Logic/Services/SourceControl/SubmitToolPerforce.cpp
+++ b/Engine/Source/Programs/SubmitTool/Private/Logic/Services/SourceControl/SubmitToolPerforce.cpp
@@ -249,7 +249,7 @@ UE::Tasks::TTask<bool> FSubmitToolPerforce::GetStream(const FString& InStream, b
 
 void FSubmitToolPerforce::GetUsers(const FOnUsersGet::FDelegate& Callback)
 {
-	if(!CachedUsersArray.IsEmpty())
+	if(UserTask.IsValid() && UserTask.IsCompleted() && !CachedUsersArray.IsEmpty())
 	{
 		Callback.ExecuteIfBound(CachedUsersArray);
 		return;
@@ -442,7 +442,7 @@ const FString FSubmitToolPerforce::GetRecentUsersFilepath() const
 
 void FSubmitToolPerforce::GetGroups(const FOnGroupsGet::FDelegate& Callback)
 {
-	if(!CachedGroupsArray.IsEmpty())
+	if(GroupTask.IsValid() && GroupTask.IsCompleted() && !CachedGroupsArray.IsEmpty())
 	{
 		Callback.ExecuteIfBound(CachedGroupsArray);
 		return;
diff --git a/Engine/Source/Programs/SubmitTool/Private/Logic/Services/SubmitToolServiceProvider.h b/Engine/Source/Programs/SubmitTool/Private/Logic/Services/SubmitToolServiceProvider.h
--- a/Engine/Source/Programs/SubmitTool/Private/Logic/Services/SubmitToolServiceProvider.h
+++ b/Engine/Source/Programs/SubmitTool/Private/Logic/Services/SubmitToolServiceProvider.h
@@ -24,10 +24,32 @@ public:
 		return nullptr; 
 	}
 
+	template<DerivedFromISubmitToolService T>
+	TSharedPtr<T> GetService(const FString& InName)
+	{
+		if (Services.Contains(InName))
+		{
+			return StaticCastSharedPtr<T>(Services[InName]);
+		}
+
+		return nullptr;
+	}
+
 	template<DerivedFromISubmitToolService T>
 	void RegisterService(TSharedRef<T> InService)
 	{
-		Services.FindOrAdd(TNameOf<T>::GetName(), StaticCastSharedPtr<ISubmitToolService>(InService.ToSharedPtr()));
+		TSharedPtr<ISubmitToolService>& Service = Services.FindOrAdd(TNameOf<T>::GetName());
+		ensureAlwaysMsgf(Service == nullptr, TEXT("Service %s was re-registered"), TNameOf<T>::GetName());
+		Service = StaticCastSharedPtr<ISubmitToolService>(InService.ToSharedPtr());
+	}
+
+
+	template<DerivedFromISubmitToolService T>
+	void RegisterService(TSharedRef<T> InService, const FString& InName)
+	{
+		TSharedPtr<ISubmitToolService>& Service = Services.FindOrAdd(InName);
+		ensureAlwaysMsgf(Service == nullptr, TEXT("Service %s was re-registered"), *InName);
+		Service = StaticCastSharedPtr<ISubmitToolService>(InService.ToSharedPtr());
 	}
 
 private:
diff --git a/Engine/Source/Programs/SubmitTool/Private/Logic/SwarmService.cpp b/Engine/Source/Programs/SubmitTool/Private/Logic/SwarmService.cpp
--- a/Engine/Source/Programs/SubmitTool/Private/Logic/SwarmService.cpp
+++ b/Engine/Source/Programs/SubmitTool/Private/Logic/SwarmService.cpp
@@ -15,11 +15,11 @@
 
 #include "Logging/SubmitToolLog.h"
 
-FSwarmService::FSwarmService(TSharedRef<FSubmitToolServiceProvider> InServiceProvider) :
+FSwarmService::FSwarmService(TWeakPtr<FSubmitToolServiceProvider> InServiceProvider) :
 	ServiceProvider(InServiceProvider)
 {
 	GetSwarmURL();
-	ServiceProvider->GetService<ISTSourceControlService>()->GetUsersAndGroups(FOnUsersAndGroupsGet::FDelegate::CreateLambda([this](TArray<TSharedPtr<FUserData>>& InUsers, TArray<TSharedPtr<FString>>& InGroups) {
+	ServiceProvider.Pin()->GetService<ISTSourceControlService>()->GetUsersAndGroups(FOnUsersAndGroupsGet::FDelegate::CreateLambda([this](TArray<TSharedPtr<FUserData>>& InUsers, TArray<TSharedPtr<FString>>& InGroups) {
 		Users = &InUsers;
 		Groups = &InGroups;
 	}));
@@ -32,7 +32,7 @@ void FSwarmService::FetchReview(const OnGetReviewComplete& OnComplete)
 		return;
 	}
 
-	FString Changelist = ServiceProvider->GetService<FChangelistService>()->GetCLID();
+	FString Changelist = ServiceProvider.Pin()->GetService<FChangelistService>()->GetCLID();
 
 	if (Changelist.IsEmpty())
 	{
@@ -41,7 +41,7 @@ void FSwarmService::FetchReview(const OnGetReviewComplete& OnComplete)
 
 	const TSharedRef<IHttpRequest> HttpRequest = FHttpModule::Get().CreateRequest();
 	HttpRequest->SetHeader(TEXT("Content-Type"), TEXT("application/json; charset=utf-8"));
-	HttpRequest->SetHeader(TEXT("Authorization"), *ServiceProvider->GetService<ISTSourceControlService>()->GetAuthTicket().ToString());
+	HttpRequest->SetHeader(TEXT("Authorization"), *ServiceProvider.Pin()->GetService<ISTSourceControlService>()->GetAuthTicket().ToString());
 
 	HttpRequest->SetURL(FString::Format(TEXT("{0}?change={1}&max={2}"),
 		{
@@ -99,7 +99,7 @@ void FSwarmService::CreateReview(const OnCreateReviewComplete& OnComplete)
 		return;
 	}
 
-	FString Changelist = ServiceProvider->GetService<FChangelistService>()->GetCLID();
+	FString Changelist = ServiceProvider.Pin()->GetService<FChangelistService>()->GetCLID();
 
 	if(Changelist.IsEmpty())
 	{
@@ -113,7 +113,7 @@ void FSwarmService::CreateReview(const OnCreateReviewComplete& OnComplete)
 
 	CreateSwarmRequest = FHttpModule::Get().CreateRequest();
 	CreateSwarmRequest->SetHeader(TEXT("Content-Type"), TEXT("application/json; charset=utf-8"));
-	CreateSwarmRequest->SetHeader(TEXT("Authorization"), *ServiceProvider->GetService<ISTSourceControlService>()->GetAuthTicket().ToString());
+	CreateSwarmRequest->SetHeader(TEXT("Authorization"), *ServiceProvider.Pin()->GetService<ISTSourceControlService>()->GetAuthTicket().ToString());
 
 	CreateSwarmRequest->SetURL(ReviewsURL());
 	CreateSwarmRequest->SetVerb(TEXT("POST"));
@@ -124,7 +124,7 @@ void FSwarmService::CreateReview(const OnCreateReviewComplete& OnComplete)
 	TArray<TSharedPtr<FJsonValue>> ReviewersObject;
 	TArray<TSharedPtr<FJsonValue>> GroupsObject;
 
-	for (const FTag* Tag : ServiceProvider->GetService<FTagService>()->GetTagsArray())
+	for (const FTag* Tag : ServiceProvider.Pin()->GetService<FTagService>()->GetTagsArray())
 	{
 		if (Tag->Definition.InputSubType.Equals(TEXT("Swarm"), ESearchCase::IgnoreCase))
 		{
@@ -166,7 +166,7 @@ void FSwarmService::CreateReview(const OnCreateReviewComplete& OnComplete)
 		RequestJson->SetArrayField(TEXT("requiredGroups"), GroupsObject);
 	}
 
-	RequestJson->SetStringField(TEXT("description"), ServiceProvider->GetService<FChangelistService>()->GetCLDescription());
+	RequestJson->SetStringField(TEXT("description"), ServiceProvider.Pin()->GetService<FChangelistService>()->GetCLDescription());
 
 	FString BodyString;
 	TSharedRef<TJsonWriter<TCHAR>> JsonWriter = TJsonWriterFactory<>::Create(&BodyString);
@@ -246,7 +246,7 @@ void FSwarmService::UpdateReviewDescription(const TDelegate<void(bool)>& OnCompl
 
 	UpdateSwarmRequest = FHttpModule::Get().CreateRequest();
 	UpdateSwarmRequest->SetHeader(TEXT("Content-Type"), TEXT("application/json; charset=utf-8"));
-	UpdateSwarmRequest->SetHeader(TEXT("Authorization"), *ServiceProvider->GetService<ISTSourceControlService>()->GetAuthTicket().ToString());
+	UpdateSwarmRequest->SetHeader(TEXT("Authorization"), *ServiceProvider.Pin()->GetService<ISTSourceControlService>()->GetAuthTicket().ToString());
 
 	FString URL = ReviewsURL() / FString::FromInt(Review->Id);
 	UpdateSwarmRequest->SetURL(URL);
@@ -303,7 +303,7 @@ const FString& FSwarmService::GetSwarmURL()
 {
 	if(SwarmURL.IsEmpty())
 	{
-		ServiceProvider->GetService<ISTSourceControlService>()->RunCommand(TEXT("property"), { "-l", "-n", "P4.Swarm.URL" }, FOnSCCCommandComplete::CreateLambda(
+		ServiceProvider.Pin()->GetService<ISTSourceControlService>()->RunCommand(TEXT("property"), { "-l", "-n", "P4.Swarm.URL" }, FOnSCCCommandComplete::CreateLambda(
 			[this](bool bSuccess, const FSCCRecordSet& InResultValues, const FSourceControlResultInfo& InResultsInfo)
 			{
 				if(bSuccess && InResultValues.Num() > 0 && InResultValues[0].Contains(TEXT("value")))
diff --git a/Engine/Source/Programs/SubmitTool/Private/Logic/SwarmService.h b/Engine/Source/Programs/SubmitTool/Private/Logic/SwarmService.h
--- a/Engine/Source/Programs/SubmitTool/Private/Logic/SwarmService.h
+++ b/Engine/Source/Programs/SubmitTool/Private/Logic/SwarmService.h
@@ -18,7 +18,7 @@ DECLARE_DELEGATE_TwoParams(OnCreateReviewComplete, bool, const FString&)
 class FSwarmService : public ISubmitToolService
 {
 public:
-	FSwarmService(TSharedRef<FSubmitToolServiceProvider> InServiceProvider);
+	FSwarmService(TWeakPtr<FSubmitToolServiceProvider> InServiceProvider);
 	void FetchReview(const OnGetReviewComplete& OnComplete);
 
 	void CreateReview(const OnCreateReviewComplete& OnComplete);
@@ -64,7 +64,7 @@ private:
 	TSharedPtr<IHttpRequest> CreateSwarmRequest;
 	TSharedPtr<IHttpRequest> UpdateSwarmRequest;
 
-	TSharedRef<FSubmitToolServiceProvider> ServiceProvider;
+	TWeakPtr<FSubmitToolServiceProvider> ServiceProvider;
 	FString SwarmURL;
 	TUniquePtr<FSwarmReview> Review = nullptr;
 
diff --git a/Engine/Source/Programs/SubmitTool/Private/Logic/TasksService.h b/Engine/Source/Programs/SubmitTool/Private/Logic/TasksService.h
--- a/Engine/Source/Programs/SubmitTool/Private/Logic/TasksService.h
+++ b/Engine/Source/Programs/SubmitTool/Private/Logic/TasksService.h
@@ -70,4 +70,3 @@ protected:
 	FTSTicker::FDelegateHandle TickerHandle;
 };
 
-Expose_TNameOf(FTasksService);
\ No newline at end of file
diff --git a/Engine/Source/Programs/SubmitTool/Private/Models/ModelInterface.cpp b/Engine/Source/Programs/SubmitTool/Private/Models/ModelInterface.cpp
--- a/Engine/Source/Programs/SubmitTool/Private/Models/ModelInterface.cpp
+++ b/Engine/Source/Programs/SubmitTool/Private/Models/ModelInterface.cpp
@@ -46,22 +46,22 @@ FModelInterface::FModelInterface(const FSubmitToolParameters& InParameters) :
 	SourceControlService = MakeShared<FSubmitToolPerforce>(InParameters);
 	ServiceProvider->RegisterService<ISTSourceControlService>(SourceControlService.ToSharedRef());
 	ValidationService = MakeShared<FTasksService>(InParameters.Validators, TEXT("SubmitTool.StandAlone.Validator"));
-	ServiceProvider->RegisterService<FTasksService>(ValidationService.ToSharedRef());
+	ServiceProvider->RegisterService<FTasksService>(ValidationService.ToSharedRef(), TEXT("ValidationService"));
 	PresubmitOperationsService = MakeShared<FTasksService>(InParameters.PresubmitOperations, TEXT("SubmitTool.StandAlone.PresubmitOperation"));
-	ServiceProvider->RegisterService<FTasksService>(PresubmitOperationsService.ToSharedRef());
+	ServiceProvider->RegisterService<FTasksService>(PresubmitOperationsService.ToSharedRef(), TEXT("PresubmitOperationsService"));
 	CredentialsService = MakeShared<FCredentialsService>(InParameters.OAuthParameters);
 	ServiceProvider->RegisterService<FCredentialsService>(CredentialsService.ToSharedRef());
 	ChangelistService = MakeShared<FChangelistService>(InParameters.GeneralParameters, SourceControlService, CLReadyCallback, CLRefreshCallback);
 	ServiceProvider->RegisterService<FChangelistService>(ChangelistService.ToSharedRef());
-	P4LockdownService = MakeShared<FP4LockdownService>(InParameters.P4LockdownParameters, ServiceProvider.ToSharedRef());
+	P4LockdownService = MakeShared<FP4LockdownService>(InParameters.P4LockdownParameters, ServiceProvider);
 	ServiceProvider->RegisterService<FP4LockdownService>(P4LockdownService.ToSharedRef());
 	TagService = MakeShared<FTagService>(InParameters, ChangelistService);
 	ServiceProvider->RegisterService<FTagService>(TagService.ToSharedRef());
-	SwarmService = MakeShared<FSwarmService>(ServiceProvider.ToSharedRef());
+	SwarmService = MakeShared<FSwarmService>(ServiceProvider);
 	ServiceProvider->RegisterService<FSwarmService>(SwarmService.ToSharedRef());
-	PreflightService = MakeShared<FPreflightService>(InParameters.HordeParameters, this, ServiceProvider.ToSharedRef());
+	PreflightService = MakeShared<FPreflightService>(InParameters.HordeParameters, this, ServiceProvider);
 	ServiceProvider->RegisterService<FPreflightService>(PreflightService.ToSharedRef());
-	JiraService = MakeShared<FJiraService>(InParameters.JiraParameters, 256, ServiceProvider.ToSharedRef());
+	JiraService = MakeShared<FJiraService>(InParameters.JiraParameters, 256, ServiceProvider);
 	ServiceProvider->RegisterService<FJiraService>(JiraService.ToSharedRef());
 	FNIntegrationService = MakeShared<FIntegrationService>(InParameters.IntegrationParameters, ServiceProvider);
 	ServiceProvider->RegisterService<FIntegrationService>(FNIntegrationService.ToSharedRef());
@@ -151,6 +151,8 @@ FModelInterface::~FModelInterface()
 	ValidationService->OnTasksQueueFinished.Remove(OnValidationFinishedHandle);
 	ValidationService->OnSingleTaskFinished.Remove(OnSingleValidationFinishedHandle);
 	PresubmitOperationsService->OnTasksQueueFinished.Remove(OnPresubmitFinishedHandle);
+
+	ServiceProvider.Reset();
 }
 
 void FModelInterface::Dispose() const
@@ -172,6 +174,7 @@ void FModelInterface::Dispose() const
 	}
 
 	ServiceProvider->GetService<ICacheDataService>()->SaveCacheToDisk();
+	ValidationService->StopTasks();
 }
 
 void FModelInterface::ParseValidators() const
diff --git a/Engine/Source/Programs/SubmitTool/Private/View/SubmitToolWindow.cpp b/Engine/Source/Programs/SubmitTool/Private/View/SubmitToolWindow.cpp
--- a/Engine/Source/Programs/SubmitTool/Private/View/SubmitToolWindow.cpp
+++ b/Engine/Source/Programs/SubmitTool/Private/View/SubmitToolWindow.cpp
@@ -26,12 +26,11 @@ TSharedRef<SDockTab> SubmitToolWindow::BuildMainTab(TSharedPtr<SWindow> InParent
 							if(!Tab->GetParentWindow()->IsWindowMaximized())
 							{
 								FSubmitToolUserPrefs::Get()->WindowPosition = Tab->GetParentWindow()->GetPositionInScreen();
-								FSubmitToolUserPrefs::Get()->WindowSize = Tab->GetParentWindow()->GetSizeInScreen();
+								FSubmitToolUserPrefs::Get()->WindowSize = Tab->GetParentWindow()->GetClientSizeInScreen();
 							}
 						})
 		.Label_Lambda([&ModelInterface = ModelInterface]() {return FText::FromString(FString::Printf(TEXT("Changelist: %s"), *ModelInterface->GetCLID())); });
 
-
 	if (ModelInterface->CheckForNewVersion())
 	{
 		CreateAutoUpdateSubmitToolContent(InParentWindow);
diff --git a/Engine/Source/Programs/SubmitTool/Private/View/Widgets/SubmitToolWidget.cpp b/Engine/Source/Programs/SubmitTool/Private/View/Widgets/SubmitToolWidget.cpp
--- a/Engine/Source/Programs/SubmitTool/Private/View/Widgets/SubmitToolWidget.cpp
+++ b/Engine/Source/Programs/SubmitTool/Private/View/Widgets/SubmitToolWidget.cpp
@@ -600,6 +600,7 @@ TSharedRef<SWidget> SubmitToolWidget::BuildOutputLogWidget()
 	SummaryLogCreationParams.DefaultCategorySelection.Emplace(LogSubmitTool.GetCategoryName(), true);
 	SummaryLogCreationParams.DefaultCategorySelection.Emplace(LogValidatorsResult.GetCategoryName(), true);
 	SummaryLogCreationParams.DefaultCategorySelection.Emplace(LogPresubmitResult.GetCategoryName(), true);
+	SummaryLogCreationParams.DefaultCategorySelection.Emplace(LogOutputDevice.GetCategoryName(), true);
 	SummaryLogCreationParams.DefaultCategorySelection.Emplace(LogSubmitToolP4.GetCategoryName(), true);
 
 	SummaryLogDockTab = SNew(SDockTab)
