Adding warning messaging when using a different PF template than suggested
Added autoupdate option enabled by default
Fixed preflight priorities config
Added message watch for UBT unhandled exception lines
Manually forcing the root volume to be uppercase to workaround a RunUBT.bat issue causing recompilation

diff --git a/Engine/Programs/SubmitTool/Config/SubmitTool.ini b/Engine/Programs/SubmitTool/Config/SubmitTool.ini
--- a/Engine/Programs/SubmitTool/Config/SubmitTool.ini
+++ b/Engine/Programs/SubmitTool/Config/SubmitTool.ini
@@ -135,7 +135,7 @@ Configurations=("Development", "Debug", "Shipping")
 Platforms=("Win64", "Linux", "Mac")
 Targets=("Editor", "Client")
 ToolTip="This Validator compiles the code in the changelist using the configuration/platform provided, it also runs Static Code Analysis at the same time."
-ErrorMessages=("A conflicting instance of", "dereferencing NULL pointer")
+ErrorMessages=("A conflicting instance of", "dereferencing NULL pointer", "Unhandled exception:")
 UIGroup="Code Validations"
 ExecutionBlockGroups=("UBT")
 TaskArea="LocalFiles"
diff --git a/Engine/Source/Programs/SubmitTool/Private/CommandLine/CmdLineParameters.h b/Engine/Source/Programs/SubmitTool/Private/CommandLine/CmdLineParameters.h
--- a/Engine/Source/Programs/SubmitTool/Private/CommandLine/CmdLineParameters.h
+++ b/Engine/Source/Programs/SubmitTool/Private/CommandLine/CmdLineParameters.h
@@ -46,6 +46,13 @@ namespace FSubmitToolCmdLine
 				}
 			}
 			OutValue.MidInline(Left, (Right - Left) + 1);
+
+			#if PLATFORM_WINDOWS
+			if (OutValue[1] == ':')
+			{
+				OutValue[0] = FChar::ToUpper(OutValue[0]);
+			}
+			#endif
 		}),
 		MakeShared<FCmdLineParameter>(ParameterFile, false, TEXT("Submit tool parameters XML override.")),
 	};
diff --git a/Engine/Source/Programs/SubmitTool/Private/Logic/JiraService.cpp b/Engine/Source/Programs/SubmitTool/Private/Logic/JiraService.cpp
--- a/Engine/Source/Programs/SubmitTool/Private/Logic/JiraService.cpp
+++ b/Engine/Source/Programs/SubmitTool/Private/Logic/JiraService.cpp
@@ -460,7 +460,7 @@ void FJiraService::CreateServiceDeskRequest(TSharedPtr<FJsonObject> InBaseJiraJs
 
 	if(!Definition.PreflightField.IsEmpty())
 	{		
-		const FTag* PreflightTag = ServiceProvider.Pin()->GetService<FTagService>()->GetTag(TEXT("#preflight"));
+		const FTag* PreflightTag = ServiceProvider.Pin()->GetService<FTagService>()->GetTagOfSubtype(TEXT("preflight"));
 		if (PreflightTag != nullptr && PreflightTag->GetValues().Num() > 0)
 		{
 			FString PreflightTagValue = PreflightTag->GetValues()[0];
diff --git a/Engine/Source/Programs/SubmitTool/Private/Logic/PreflightService.cpp b/Engine/Source/Programs/SubmitTool/Private/Logic/PreflightService.cpp
--- a/Engine/Source/Programs/SubmitTool/Private/Logic/PreflightService.cpp
+++ b/Engine/Source/Programs/SubmitTool/Private/Logic/PreflightService.cpp
@@ -36,7 +36,7 @@ FPreflightService::FPreflightService(
 	DeleteShelveCallback = FSourceControlOperationComplete::CreateRaw(this, &FPreflightService::OnDeleteShelveOperationComplete);
 	CreateShelveCallback = FSourceControlOperationComplete::CreateRaw(this, &FPreflightService::OnCreateShelveOperationComplete);
 
-	PreflightTag = ServiceProvider.Pin()->GetService<FTagService>()->GetTag(TEXT("#preflight"));
+	PreflightTag = ServiceProvider.Pin()->GetService<FTagService>()->GetTagOfSubtype(TEXT("preflight"));
 }
 
 FPreflightService::~FPreflightService()
diff --git a/Engine/Source/Programs/SubmitTool/Private/Logic/PreflightService.h b/Engine/Source/Programs/SubmitTool/Private/Logic/PreflightService.h
--- a/Engine/Source/Programs/SubmitTool/Private/Logic/PreflightService.h
+++ b/Engine/Source/Programs/SubmitTool/Private/Logic/PreflightService.h
@@ -60,6 +60,12 @@ public:
 	{
 		return Definition.HordeServerAddress;
 	}
+	const FString& GetDefaultPreflightTemplate() const
+	{
+		return Definition.DefaultPreflightTemplate;
+	}
+
+	bool SelectPreflightTemplate(FPreflightTemplateDefinition& OutTemplate) const;
 
 	FOnPreflightDataUpdated OnPreflightDataUpdated;
 
@@ -69,7 +75,6 @@ private:
 	void Requeue();
 
 	void StartPreflight();
-	bool SelectPreflightTemplate(FPreflightTemplateDefinition& OutTemplate) const;
 	FString GetAdditionalTasksString(const FPreflightTemplateDefinition& InTemplate) const;
 
 	void FetchUnlinkedPreflight(const FString& InPreflightId, bool bRequeue, const FString& InOAuthToken);
diff --git a/Engine/Source/Programs/SubmitTool/Private/Logic/TagService.h b/Engine/Source/Programs/SubmitTool/Private/Logic/TagService.h
--- a/Engine/Source/Programs/SubmitTool/Private/Logic/TagService.h
+++ b/Engine/Source/Programs/SubmitTool/Private/Logic/TagService.h
@@ -34,6 +34,19 @@ public:
 		return nullptr;
 	};
 
+	FTag* GetTagOfSubtype(const FString& InTagSubtype)
+	{
+		for (const TPair<FString, FTag>& Tag : RegisteredTags)
+		{
+			if (Tag.Value.Definition.InputSubType.Equals(InTagSubtype, ESearchCase::IgnoreCase))
+			{
+				return &RegisteredTags[Tag.Key];
+			}
+		}
+
+		return nullptr;
+	}
+
 	FTagUpdated OnTagUpdated;
 
 	const TArray<const FTag*>& GetTagsArray() const;
diff --git a/Engine/Source/Programs/SubmitTool/Private/Logic/Validators/PreflightValidator.cpp b/Engine/Source/Programs/SubmitTool/Private/Logic/Validators/PreflightValidator.cpp
--- a/Engine/Source/Programs/SubmitTool/Private/Logic/Validators/PreflightValidator.cpp
+++ b/Engine/Source/Programs/SubmitTool/Private/Logic/Validators/PreflightValidator.cpp
@@ -12,7 +12,7 @@
 bool FPreflightValidator::Activate()
 {
 	bIsValidSetup = FValidatorBase::Activate();
-	PreflightTag = ServiceProvider.Pin()->GetService<FTagService>()->GetTag(TEXT("#preflight"));
+	PreflightTag = ServiceProvider.Pin()->GetService<FTagService>()->GetTagOfSubtype(TEXT("preflight"));
 
 	bIsValidSetup = bIsValidSetup && PreflightTag != nullptr;
 
@@ -51,6 +51,7 @@ bool FPreflightValidator::Validate(const FString& InCLDescription, const TArray<
 
 	if(PreflightService->GetPreflightData().IsValid())
 	{
+		PreflightService->SelectPreflightTemplate(SuggestedTemplate);
 		ValidatePreflights(PreflightService->GetPreflightData(), PreflightService->GetUnlinkedPreflights());
 	}
 
@@ -207,6 +208,12 @@ void FPreflightValidator::ValidatePreflights(const TUniquePtr<FPreflightList>& I
 
 		for(const FPreflightData* PreflightPtr : ValidPreflights)
 		{
+			if (SuggestedTemplate.Template != ServiceProvider.Pin()->GetService<FPreflightService>()->GetDefaultPreflightTemplate() && PreflightPtr->TemplateId != SuggestedTemplate.Template)
+			{
+				UE_LOG(LogValidators, Warning, TEXT("[%s] %s preflight used template %s, submit tool recommended preflight for your CL was %s, make sure your changes are covered by the %s preflight"), *GetValidatorName(), *PreflightPtr->ID, *PreflightPtr->TemplateId, *SuggestedTemplate.Template, *PreflightPtr->TemplateId);
+				UE_LOG(LogValidatorsResult, Warning, TEXT("[%s] %s preflight used template %s, submit tool recommended preflight for your CL was %s, make sure your changes are covered by the %s preflight"), *GetValidatorName(), *PreflightPtr->ID, *PreflightPtr->TemplateId, *SuggestedTemplate.Template, *PreflightPtr->TemplateId);
+			}
+
 			UE_LOG(LogValidators, Log, TEXT("[%s] %s preflight is valid and has succeeded"), *GetValidatorName(), *PreflightPtr->ID);
 			UE_LOG(LogValidatorsResult, Log, TEXT("[%s] %s preflight is valid and has succeeded"), *GetValidatorName(), *PreflightPtr->ID);
 		}
diff --git a/Engine/Source/Programs/SubmitTool/Private/Logic/Validators/PreflightValidator.h b/Engine/Source/Programs/SubmitTool/Private/Logic/Validators/PreflightValidator.h
--- a/Engine/Source/Programs/SubmitTool/Private/Logic/Validators/PreflightValidator.h
+++ b/Engine/Source/Programs/SubmitTool/Private/Logic/Validators/PreflightValidator.h
@@ -24,6 +24,8 @@ protected:
 
 	FTag* PreflightTag; 
 
+	FPreflightTemplateDefinition SuggestedTemplate;
+
 	virtual void ValidationFinished(bool bSuccess) override;
 	FDelegateHandle PreflightUpdateHandler;
 	FDelegateHandle TagUpdateHandler;
diff --git a/Engine/Source/Programs/SubmitTool/Private/Models/PreflightData.h b/Engine/Source/Programs/SubmitTool/Private/Models/PreflightData.h
--- a/Engine/Source/Programs/SubmitTool/Private/Models/PreflightData.h
+++ b/Engine/Source/Programs/SubmitTool/Private/Models/PreflightData.h
@@ -111,6 +111,8 @@ struct FPreflightData
 	UPROPERTY()
 	FString Name;
 	UPROPERTY()
+	FString TemplateId;
+	UPROPERTY()
 	FDateTime CreateTime;
 	UPROPERTY()
 	FDateTime UpdateTime;
diff --git a/Engine/Source/Programs/SubmitTool/Private/Models/SubmitToolUserPrefs.h b/Engine/Source/Programs/SubmitTool/Private/Models/SubmitToolUserPrefs.h
--- a/Engine/Source/Programs/SubmitTool/Private/Models/SubmitToolUserPrefs.h
+++ b/Engine/Source/Programs/SubmitTool/Private/Models/SubmitToolUserPrefs.h
@@ -40,6 +40,9 @@ struct FSubmitToolUserPrefs
 	TMap<FName, bool> UISectionExpandState =
 	{};
 
+	UPROPERTY()
+	bool bAutoUpdate = true;
+
 	virtual ~FSubmitToolUserPrefs();
 
 	static FSubmitToolUserPrefs* Get()
diff --git a/Engine/Source/Programs/SubmitTool/Private/Version/AppVersionDefines.h b/Engine/Source/Programs/SubmitTool/Private/Version/AppVersionDefines.h
--- a/Engine/Source/Programs/SubmitTool/Private/Version/AppVersionDefines.h
+++ b/Engine/Source/Programs/SubmitTool/Private/Version/AppVersionDefines.h
@@ -21,7 +21,7 @@
 
 #define SUBMIT_TOOL_VERSION_MAJOR 5
 #define SUBMIT_TOOL_VERSION_MINOR 1
-#define SUBMIT_TOOL_VERSION_PATCH 3
+#define SUBMIT_TOOL_VERSION_PATCH 4
 
 #define SUBMIT_TOOL_VERSION_STRING \
 	VERSION_STRINGIFY(SUBMIT_TOOL_VERSION_MAJOR)\
diff --git a/Engine/Source/Programs/SubmitTool/Private/View/SubmitToolWindow.cpp b/Engine/Source/Programs/SubmitTool/Private/View/SubmitToolWindow.cpp
--- a/Engine/Source/Programs/SubmitTool/Private/View/SubmitToolWindow.cpp
+++ b/Engine/Source/Programs/SubmitTool/Private/View/SubmitToolWindow.cpp
@@ -34,6 +34,10 @@ TSharedRef<SDockTab> SubmitToolWindow::BuildMainTab(TSharedPtr<SWindow> InParent
 	if (ModelInterface->CheckForNewVersion())
 	{
 		CreateAutoUpdateSubmitToolContent(InParentWindow);
+		if (FSubmitToolUserPrefs::Get()->bAutoUpdate)
+		{
+			ModelInterface->InstallLatestVersion();
+		}
 	}
 	else
 	{
diff --git a/Engine/Source/Programs/SubmitTool/Private/View/Widgets/AutoUpdateWidget.cpp b/Engine/Source/Programs/SubmitTool/Private/View/Widgets/AutoUpdateWidget.cpp
--- a/Engine/Source/Programs/SubmitTool/Private/View/Widgets/AutoUpdateWidget.cpp
+++ b/Engine/Source/Programs/SubmitTool/Private/View/Widgets/AutoUpdateWidget.cpp
@@ -5,9 +5,11 @@
 #include "Widgets/Input/SButton.h"
 #include "Widgets/Input/STextComboBox.h"
 #include "Widgets/Input/SEditableTextBox.h"
+#include "Widgets/Input/SCheckBox.h"
 #include "Widgets/Layout/SUniformGridPanel.h"
 #include "Widgets/SToolTip.h"
 
+#include "Models/SubmitToolUserPrefs.h"
 #include "Models/ModelInterface.h"
 #include "View/SubmitToolStyle.h"
 
@@ -164,11 +166,43 @@ void SAutoUpdateWidget::Construct(const FArguments& InArgs)
 					SNew(STextBlock)
 					.MinDesiredWidth(130)
 					.Justification(ETextJustify::Center)
-					.Text(FText::FromString(TEXT("Use Current Version")))
+					.Text_Lambda([this](){ return ModelInterface->GetDownloadMessage().IsEmpty() ? FText::FromString(TEXT("Use Current Version")) : FText::FromString(TEXT("Cancel Download")); })
 				]
 			]
 		];
 
+	Contents->AddSlot()
+		.AutoHeight()
+		[
+			SNew(SHorizontalBox)
+				+ SHorizontalBox::Slot()
+				.Padding(3, 3)
+				.AutoWidth()
+				.VAlign(VAlign_Center)
+				[
+					SNew(SCheckBox)
+					.IsChecked_Lambda([this](){ return FSubmitToolUserPrefs::Get()->bAutoUpdate ? ECheckBoxState::Checked : ECheckBoxState::Unchecked;})
+					.OnCheckStateChanged_Lambda([this](ECheckBoxState InCheckBoxState){
+						FSubmitToolUserPrefs::Get()->bAutoUpdate = !FSubmitToolUserPrefs::Get()->bAutoUpdate;
+					})
+				]
+				+SHorizontalBox::Slot()
+				.AutoWidth()
+				.VAlign(VAlign_Center)
+				[
+					SNew(SButton)
+					.ButtonStyle(FAppStyle::Get(), "InvisibleButton")
+					.IsFocusable(false)	
+					.OnClicked_Lambda([]() { FSubmitToolUserPrefs::Get()->bAutoUpdate = !FSubmitToolUserPrefs::Get()->bAutoUpdate; return FReply::Handled(); })
+					[
+						SNew(STextBlock)
+						.Justification(ETextJustify::Left)
+						.Font(FAppStyle::Get().GetFontStyle("SmallFont"))
+						.MinDesiredWidth(60)
+						.Text(FText::FromString(TEXT("Auto Update")))
+					]
+				]
+		];
 
 	Contents->AddSlot()
 		.AutoHeight()
diff --git a/Engine/Source/Programs/SubmitTool/Private/View/Widgets/TagWidget.cpp b/Engine/Source/Programs/SubmitTool/Private/View/Widgets/TagWidget.cpp
--- a/Engine/Source/Programs/SubmitTool/Private/View/Widgets/TagWidget.cpp
+++ b/Engine/Source/Programs/SubmitTool/Private/View/Widgets/TagWidget.cpp
@@ -156,7 +156,7 @@ void STagWidget::Construct(const FArguments& InArgs)
 			SNew(SEditableTextBox)
 			.MinDesiredWidth(256)
 			.OnTextCommitted_Lambda([this](const FText& Text, ETextCommit::Type CommitType) { if(CommitType != ETextCommit::OnEnter) {
-				if(Tag->Definition.TagLabel.Equals("Preflight", ESearchCase::IgnoreCase))
+				if(Tag->Definition.InputSubType.Equals("Preflight", ESearchCase::IgnoreCase))
 				{
 					ModelInterface->RefreshPreflightInformation();
 				}
