Using the right stream to test for preflight suggestions for virtual streams

diff --git a/Engine/Source/Programs/SubmitTool/Private/Logic/PreflightService.cpp b/Engine/Source/Programs/SubmitTool/Private/Logic/PreflightService.cpp
--- a/Engine/Source/Programs/SubmitTool/Private/Logic/PreflightService.cpp
+++ b/Engine/Source/Programs/SubmitTool/Private/Logic/PreflightService.cpp
@@ -555,7 +555,6 @@ void FPreflightService::StartPreflight()
 		}
 	}
 
-
 	FString StartPreflightUrl = FString::Format(*Definition.StartPreflightURLFormat, GetFormatParameters());
 
 	// If for some reason, our preflight settings are missing, this will be empty, let's not popup a browser with nothing in it
@@ -621,8 +620,6 @@ void FPreflightService::OnCreateShelveOperationComplete(const FSourceControlOper
 
 bool FPreflightService::SelectPreflightTemplate(FPreflightTemplateDefinition& OutTemplate) const
 {
-	FConfiguration::AddOrUpdateEntry(TEXT("$(StreamRoot)"),	ServiceProvider.Pin()->GetService<ISTSourceControlService>()->GetRootStreamName());
-
 	TSharedPtr<FChangelistService> ChangelistService = ServiceProvider.Pin()->GetService<FChangelistService>();
 
 	// Get the list of files in the changelist
@@ -631,7 +628,7 @@ bool FPreflightService::SelectPreflightTemplate(FPreflightTemplateDefinition& Ou
 	// Loop through each definition to see if the files are in the path then check extension
 	for(const FPreflightTemplateDefinition& Def : Definition.Definitions)
 	{
-		FString RegexPat = FConfiguration::Substitute(Def.RegexPath);
+		FString RegexPat = Def.RegexPath.Replace(TEXT("$(StreamRoot)"), *StreamName, ESearchCase::IgnoreCase);
 		FRegexPattern Pattern = FRegexPattern(RegexPat, ERegexPatternFlags::CaseInsensitive);
 		for(const FString& File : FilesInCl)
 		{
diff --git a/Engine/Source/Programs/SubmitTool/Private/Version/AppVersionDefines.h b/Engine/Source/Programs/SubmitTool/Private/Version/AppVersionDefines.h
--- a/Engine/Source/Programs/SubmitTool/Private/Version/AppVersionDefines.h
+++ b/Engine/Source/Programs/SubmitTool/Private/Version/AppVersionDefines.h
@@ -21,7 +21,7 @@
 
 #define SUBMIT_TOOL_VERSION_MAJOR 5
 #define SUBMIT_TOOL_VERSION_MINOR 1
-#define SUBMIT_TOOL_VERSION_PATCH 4
+#define SUBMIT_TOOL_VERSION_PATCH 5
 
 #define SUBMIT_TOOL_VERSION_STRING \
 	VERSION_STRINGIFY(SUBMIT_TOOL_VERSION_MAJOR)\
