Add error code to log output when building UVT via UBT fails.

#rnx

diff --git a/Engine/Source/Programs/SubmitTool/Private/Logic/PreSubmitOperations/VirtualizationOperation.cpp b/Engine/Source/Programs/SubmitTool/Private/Logic/PreSubmitOperations/VirtualizationOperation.cpp
--- a/Engine/Source/Programs/SubmitTool/Private/Logic/PreSubmitOperations/VirtualizationOperation.cpp
+++ b/Engine/Source/Programs/SubmitTool/Private/Logic/PreSubmitOperations/VirtualizationOperation.cpp
@@ -279,7 +279,7 @@ void FVirtualizationOperation::Tick(float InDeltatime)
 			}
 			else
 			{
-				LogFailure(FString::Printf(TEXT("[%s] Failed to build Virtualization tool - %.1f(s)"), *ValidatorName, TotalCompileTime));
+				LogFailure(FString::Printf(TEXT("[%s] Failed to build Virtualization tool (error code %d) - %.1f(s)"), *ValidatorName, CompileResult, TotalCompileTime));
 				ErrorCode = EVirtualizationErrorCode::CompileFailed;
 				
 				ValidationFinished(false);
