Fix for Preflight validator waiting for optional validators

diff --git a/Engine/Source/Programs/SubmitTool/Private/Models/ModelInterface.cpp b/Engine/Source/Programs/SubmitTool/Private/Models/ModelInterface.cpp
--- a/Engine/Source/Programs/SubmitTool/Private/Models/ModelInterface.cpp
+++ b/Engine/Source/Programs/SubmitTool/Private/Models/ModelInterface.cpp
@@ -309,7 +309,7 @@ bool FModelInterface::CanLaunchPreflight() const
 			const TSharedPtr<const FValidatorBase> Pinned = Validator.Pin();
 			if((Pinned->Definition->TaskArea & ETaskArea::ShelveAndLocalFiles) != ETaskArea::None)
 			{
-				if(!Pinned->GetHasPassed())
+				if(!Pinned->GetHasPassed() || (!Pinned->Definition->IsRequired && !Pinned->GetIsRunningOrQueued()))
 				{
 					return false;
 				}
