bumping submit tool version

diff --git a/Engine/Source/Programs/SubmitTool/Private/Version/AppVersionDefines.h b/Engine/Source/Programs/SubmitTool/Private/Version/AppVersionDefines.h
--- a/Engine/Source/Programs/SubmitTool/Private/Version/AppVersionDefines.h
+++ b/Engine/Source/Programs/SubmitTool/Private/Version/AppVersionDefines.h
@@ -21,7 +21,7 @@
 
 #define SUBMIT_TOOL_VERSION_MAJOR 5
 #define SUBMIT_TOOL_VERSION_MINOR 2
-#define SUBMIT_TOOL_VERSION_PATCH 0
+#define SUBMIT_TOOL_VERSION_PATCH 1
 
 #define SUBMIT_TOOL_VERSION_STRING \
 	VERSION_STRINGIFY(SUBMIT_TOOL_VERSION_MAJOR)\
