Fix for print error summaries for optional validators

diff --git a/Engine/Source/Programs/SubmitTool/Private/Logic/TasksService.cpp b/Engine/Source/Programs/SubmitTool/Private/Logic/TasksService.cpp
--- a/Engine/Source/Programs/SubmitTool/Private/Logic/TasksService.cpp
+++ b/Engine/Source/Programs/SubmitTool/Private/Logic/TasksService.cpp
@@ -401,10 +401,7 @@ bool FTasksService::Tick(float InDeltaTime)
 
 		if(!bLastRunningTasks)
 		{
-			if(!bCurrentCLValidState)
-			{
-				PrintErrorSummary();
-			}
+			PrintErrorSummary();			
 
 			TArray<FAnalyticsEventAttribute> ValidationRunResults = MakeAnalyticsEventAttributeArray(TEXT("Success"), bCurrentCLValidState);
 
