Fix preflights being blocked by [REDACTED_PERSON] and other optional validators
Cross CL Validator will now only check the uncontrolled changelist when the CL contains uassets, umaps or uplugins
Path slashes normalization

#rnx

diff --git a/Engine/Source/Programs/SubmitTool/Private/CommandLine/CmdLineParameters.h b/Engine/Source/Programs/SubmitTool/Private/CommandLine/CmdLineParameters.h
--- a/Engine/Source/Programs/SubmitTool/Private/CommandLine/CmdLineParameters.h
+++ b/Engine/Source/Programs/SubmitTool/Private/CommandLine/CmdLineParameters.h
@@ -4,6 +4,7 @@
 
 #include "CoreMinimal.h"
 #include "CmdLineParameter.h"
+#include "Misc/Paths.h"
 
 namespace FSubmitToolCmdLine
 {
@@ -53,6 +54,8 @@ namespace FSubmitToolCmdLine
 				OutValue[0] = FChar::ToUpper(OutValue[0]);
 			}
 			#endif
+			
+			FPaths::NormalizeDirectoryName(OutValue);
 		}),
 		MakeShared<FCmdLineParameter>(ParameterFile, false, TEXT("Submit tool parameters XML override.")),
 	};
diff --git a/Engine/Source/Programs/SubmitTool/Private/Logic/Validators/CrossChangelistValidator.cpp b/Engine/Source/Programs/SubmitTool/Private/Logic/Validators/CrossChangelistValidator.cpp
--- a/Engine/Source/Programs/SubmitTool/Private/Logic/Validators/CrossChangelistValidator.cpp
+++ b/Engine/Source/Programs/SubmitTool/Private/Logic/Validators/CrossChangelistValidator.cpp
@@ -16,7 +16,29 @@ bool FCrossChangelistValidator::Validate(const FString& InCLDescription, const T
 	bool bValid = true;
 
 	bValid &= CheckHeaderAndCppInDifferentChangelist(OtherChangelistsStates);
-	bValid &= CheckForFilesInUncontrolledCLFile();
+	
+	bool bCLHasAssets = false;
+	const FString UAssetExt(TEXT(".uasset"));
+	const FString UMapExt(TEXT(".umap"));
+	const FString UPluginExt(TEXT(".uplugin"));
+	for (FSourceControlStateRef FileInCL : InFilteredFilesInCL)
+	{
+		const FString& Filename = FileInCL->GetFilename();
+
+		if (Filename.EndsWith(UAssetExt, ESearchCase::IgnoreCase)
+			|| Filename.EndsWith(UMapExt, ESearchCase::IgnoreCase)
+			|| Filename.EndsWith(UPluginExt, ESearchCase::IgnoreCase))
+		{
+			bCLHasAssets = true;
+			break;
+		}
+	}
+
+	if (bCLHasAssets)
+	{
+		// Only check for uncontrolled CL files when having assets in the CL we want to submit
+		bValid &= CheckForFilesInUncontrolledCLFile();
+	}
 
 	ValidationFinished(bValid);
 	return true;
diff --git a/Engine/Source/Programs/SubmitTool/Private/Models/ModelInterface.cpp b/Engine/Source/Programs/SubmitTool/Private/Models/ModelInterface.cpp
--- a/Engine/Source/Programs/SubmitTool/Private/Models/ModelInterface.cpp
+++ b/Engine/Source/Programs/SubmitTool/Private/Models/ModelInterface.cpp
@@ -309,7 +309,7 @@ bool FModelInterface::CanLaunchPreflight() const
 			const TSharedPtr<const FValidatorBase> Pinned = Validator.Pin();
 			if((Pinned->Definition->TaskArea & ETaskArea::ShelveAndLocalFiles) != ETaskArea::None)
 			{
-				if(!Pinned->GetHasPassed() || (!Pinned->Definition->IsRequired && !Pinned->GetIsRunningOrQueued()))
+				if((Pinned->Definition->IsRequired && !Pinned->GetHasPassed()) || (!Pinned->Definition->IsRequired && Pinned->GetIsRunningOrQueued()))
 				{
 					return false;
 				}
diff --git a/Engine/Source/Programs/SubmitTool/Private/SubmitToolUtils.cpp b/Engine/Source/Programs/SubmitTool/Private/SubmitToolUtils.cpp
--- a/Engine/Source/Programs/SubmitTool/Private/SubmitToolUtils.cpp
+++ b/Engine/Source/Programs/SubmitTool/Private/SubmitToolUtils.cpp
@@ -18,7 +18,8 @@
 FString FSubmitToolUtils::GetLocalAppDataPath()
 {
 #if PLATFORM_WINDOWS
-	const FString LocalAppData = FPlatformMisc::GetEnvironmentVariable(TEXT("LOCALAPPDATA"));
+	FString LocalAppData = FPlatformMisc::GetEnvironmentVariable(TEXT("LOCALAPPDATA"));
+	FPaths::NormalizeDirectoryName(LocalAppData);
 #elif PLATFORM_MAC
 	const FString LocalAppData = FPaths::Combine(FPlatformMisc::GetEnvironmentVariable(TEXT("HOME")), TEXT("Library"), TEXT("Application Support"));
 #elif PLATFORM_LINUX
diff --git a/Engine/Source/Programs/SubmitTool/Private/View/Widgets/SubmitToolWidget.cpp b/Engine/Source/Programs/SubmitTool/Private/View/Widgets/SubmitToolWidget.cpp
--- a/Engine/Source/Programs/SubmitTool/Private/View/Widgets/SubmitToolWidget.cpp
+++ b/Engine/Source/Programs/SubmitTool/Private/View/Widgets/SubmitToolWidget.cpp
@@ -363,7 +363,6 @@ TSharedRef<SHorizontalBox> SubmitToolWidget::BuildButtonRow()
 			SNew(SButton)
 				.ButtonStyle(FAppStyle::Get(), "Button")
 				.OnClicked(this, &SubmitToolWidget::CopyAllLogsClicked)
-				.IsEnabled_Lambda([this]() { return !ModelInterface->IsValidationRunning(); })
 				[
 					SNew(STextBlock)
 						.MinDesiredWidth(130)
@@ -893,6 +892,8 @@ FReply SubmitToolWidget::CopyAllLogsClicked()
 	}
 
 	FSubmitToolUtils::CopyDiagnosticFilesToClipboard(Files);
+
+	UE_LOG(LogSubmitTool,Display, TEXT("Log files have been copied to the clipboard"));
 	return FReply::Handled();
 }
 #endif
