Do not print the error summary header/footer if a validator was cancelled as the validator will not print any of it's error messages anyway.

#rnx

diff --git a/Engine/Source/Programs/SubmitTool/Private/Logic/TasksService.cpp b/Engine/Source/Programs/SubmitTool/Private/Logic/TasksService.cpp
--- a/Engine/Source/Programs/SubmitTool/Private/Logic/TasksService.cpp
+++ b/Engine/Source/Programs/SubmitTool/Private/Logic/TasksService.cpp
@@ -174,7 +174,7 @@ void FTasksService::PrintErrorSummary()
 	bool bHeaderPrinted = false;
 	for(const TPair<FName, TSharedPtr<FValidatorBase>>& TaskPair : Tasks)
 	{
-		if (TaskPair.Value->GetErrorList().Num() != 0)
+		if (TaskPair.Value->CanPrintErrors())
 		{
 			if (!bHeaderPrinted)
 			{
diff --git a/Engine/Source/Programs/SubmitTool/Private/Logic/Validators/ValidatorBase.cpp b/Engine/Source/Programs/SubmitTool/Private/Logic/Validators/ValidatorBase.cpp
--- a/Engine/Source/Programs/SubmitTool/Private/Logic/Validators/ValidatorBase.cpp
+++ b/Engine/Source/Programs/SubmitTool/Private/Logic/Validators/ValidatorBase.cpp
@@ -270,28 +270,38 @@ void FValidatorBase::SetSelectedOption(const FString& InOptionName, const FStrin
 	OptionsProvider.SetSelectedOption(InOptionName, InOptionValue);
 }
 
+bool FValidatorBase::CanPrintErrors() const
+{
+	if (State == EValidationStates::Failed || State == EValidationStates::Timeout)
+	{
+		return !ErrorListCache.IsEmpty();
+	}
+	else
+	{
+		return false;
+	}
+}
+
 void FValidatorBase::PrintErrorSummary() const
 {
-	if(State == EValidationStates::Failed || State == EValidationStates::Timeout)
+	FScopeLock Lock(&Mutex);
+
+	if (CanPrintErrors())
 	{
-		FScopeLock Lock(&Mutex);
-		if(ErrorListCache.Num() > 0)
+		if (Definition->IsRequired)
 		{
-			if (Definition->IsRequired)
+			for(const FString& ErrorStr : ErrorListCache)
 			{
-				for(const FString& ErrorStr : ErrorListCache)
-				{
-					UE_LOG(LogValidators, Error, TEXT("%s"), *ErrorStr);
-					UE_LOG(LogValidatorsResult, Error, TEXT("%s"), *ErrorStr);
-				}
+				UE_LOG(LogValidators, Error, TEXT("%s"), *ErrorStr);
+				UE_LOG(LogValidatorsResult, Error, TEXT("%s"), *ErrorStr);
 			}
-			else
+		}
+		else
+		{
+			for (const FString& ErrorStr : ErrorListCache)
 			{
-				for (const FString& ErrorStr : ErrorListCache)
-				{
-					UE_LOG(LogValidators, Warning, TEXT("%s"), *ErrorStr);
-					UE_LOG(LogValidatorsResult, Warning, TEXT("%s"), *ErrorStr);
-				}				
+				UE_LOG(LogValidators, Warning, TEXT("%s"), *ErrorStr);
+				UE_LOG(LogValidatorsResult, Warning, TEXT("%s"), *ErrorStr);
 			}
 		}
 	}
diff --git a/Engine/Source/Programs/SubmitTool/Private/Logic/Validators/ValidatorBase.h b/Engine/Source/Programs/SubmitTool/Private/Logic/Validators/ValidatorBase.h
--- a/Engine/Source/Programs/SubmitTool/Private/Logic/Validators/ValidatorBase.h
+++ b/Engine/Source/Programs/SubmitTool/Private/Logic/Validators/ValidatorBase.h
@@ -171,12 +171,9 @@ public:
 	void SetSelectedOption(const FString& InOptionName, const FString& InOptionValue);
 
 	virtual const TArray<FAnalyticsEventAttribute> GetTelemetryAttributes() const;
-	virtual void PrintErrorSummary() const;
-
-	const TArray<FString>& GetErrorList() const
-	{
-		return ErrorListCache;
-	}
+	
+	bool CanPrintErrors() const;
+	void PrintErrorSummary() const;
 
 	const FString GetValidationConfigId() const;
 
