When logging the contents of a changelist, prefix each file path with a tab rather than newline character.

#rnx

diff --git a/Engine/Source/Programs/SubmitTool/Private/Logic/ChangelistService.cpp b/Engine/Source/Programs/SubmitTool/Private/Logic/ChangelistService.cpp
--- a/Engine/Source/Programs/SubmitTool/Private/Logic/ChangelistService.cpp
+++ b/Engine/Source/Programs/SubmitTool/Private/Logic/ChangelistService.cpp
@@ -628,7 +628,7 @@ void FChangelistService::PrintFilesAndShelvedFiles()
 		UE_LOG(LogSubmitToolP4Debug, Log, TEXT("Files in CL:"));
 		for(const FSourceControlStateRef& File : FilesInCL)
 		{
-			UE_LOG(LogSubmitToolP4Debug, Log, TEXT("\n%s"), *File->GetFilename());
+			UE_LOG(LogSubmitToolP4Debug, Log, TEXT("\t%s"), *File->GetFilename());
 		}
 	}
 
@@ -637,7 +637,7 @@ void FChangelistService::PrintFilesAndShelvedFiles()
 		UE_LOG(LogSubmitToolP4Debug, Log, TEXT("Shelved Files in CL:"));
 		for(const FSourceControlStateRef& File : ShelvedFilesInCL)
 		{
-			UE_LOG(LogSubmitToolP4Debug, Log, TEXT("\n%s"), *File->GetFilename());
+			UE_LOG(LogSubmitToolP4Debug, Log, TEXT("\t%s"), *File->GetFilename());
 		}
 	}
 }
