Submitting Josh changes for Submit Tool
Fixed a couple of small bugs in SubmitTool:
- When parsing -//Foo, it was turning into /Foo (removed one too many character)
- Removed the NoExpand flag on the custom ini hierarchy on layers with {PLATFORM}

diff --git a/Engine/Source/Programs/SubmitTool/Private/Logic/P4LockdownService.cpp b/Engine/Source/Programs/SubmitTool/Private/Logic/P4LockdownService.cpp
--- a/Engine/Source/Programs/SubmitTool/Private/Logic/P4LockdownService.cpp
+++ b/Engine/Source/Programs/SubmitTool/Private/Logic/P4LockdownService.cpp
@@ -304,7 +304,7 @@ void FP4LockdownService::ParseAllowListData()
 					if(ViewLine.StartsWith(TEXT("-"), ESearchCase::CaseSensitive))
 					{
 						bIsLocked = false;
-						ViewLine.MidInline(2);
+						ViewLine.MidInline(1);
 					}
 
 					RegexEscapeInline(ViewLine);
@@ -333,4 +333,4 @@ FString FP4LockdownService::GetFilePath(const FString& InConfigId)
 	FPaths::NormalizeFilename(LocalFilePath);
 		
 	return LocalFilePath;
-}
\ No newline at end of file
+}
diff --git a/Engine/Source/Programs/SubmitTool/Private/Parameters/SubmitToolParametersBuilder.cpp b/Engine/Source/Programs/SubmitTool/Private/Parameters/SubmitToolParametersBuilder.cpp
--- a/Engine/Source/Programs/SubmitTool/Private/Parameters/SubmitToolParametersBuilder.cpp
+++ b/Engine/Source/Programs/SubmitTool/Private/Parameters/SubmitToolParametersBuilder.cpp
@@ -27,7 +27,7 @@ FSubmitToolParametersBuilder::FSubmitToolParametersBuilder(const FString& InPara
 	ConfigHierarchy.Add({ TEXT("RootBase"), **RootBase, EConfigLayerFlags::NoExpand });
 
 	FString* RootPlatform = new FString(FPaths::Combine(RootDir, TEXT("/Config/{PLATFORM}/{PLATFORM}{TYPE}.ini")));
-	ConfigHierarchy.Add({ TEXT("RootPlatform"), **RootPlatform, EConfigLayerFlags::NoExpand });
+	ConfigHierarchy.Add({ TEXT("RootPlatform"), **RootPlatform });
 
 	if(!RootDir.IsEmpty())
 	{
@@ -50,7 +50,7 @@ FSubmitToolParametersBuilder::FSubmitToolParametersBuilder(const FString& InPara
 							ConfigHierarchy.Add( { TEXT("Project"), **ProjectIni, EConfigLayerFlags::NoExpand } );
 
 							FString* ProjectPlatform = new FString(FPaths::Combine(Dir, TEXT("/Config/{PLATFORM}/{PLATFORM}{TYPE}.ini")));
-							ConfigHierarchy.Add({ TEXT("ProjectPlatform"), **ProjectPlatform, EConfigLayerFlags::NoExpand });
+							ConfigHierarchy.Add({ TEXT("ProjectPlatform"), **ProjectPlatform });
 
 							// Let's hold on to this info for the preflight parameters
 							for (const FString& UProjectName : UProjects)
