Add a new validator to the SubmitTool.

#rnx

- The new validator is called 'PackageDataValidator' as even though it currently only cheks for packages with virtualized data it could be exapnded to check other common package format issues in the future if we want it to.
- It can be configured by providing a list of project file extensions that do not allow virtualized packages. The validator will then check each package file in the provided changelist to see if a) they are part of a project using that extension type b) contains virtualized data and if so, log an error for that package file.
- Even though we only launch a single task per validation pass, the validator has been written so that it is trivial to add new tasks.
- Include the virtualization module in SubmitTool so that we can use the utilities.

diff --git a/Engine/Source/Programs/SubmitTool/Private/Logic/Validators/PackageDataValidator.cpp b/Engine/Source/Programs/SubmitTool/Private/Logic/Validators/PackageDataValidator.cpp
new file mode 100644
--- /dev/null
+++ b/Engine/Source/Programs/SubmitTool/Private/Logic/Validators/PackageDataValidator.cpp
@@ -0,0 +1,159 @@
+// Copyright Epic Games, Inc. All Rights Reserved.
+
+#include "PackageDataValidator.h"
+
+#include "Models/ModelInterface.h"
+#include "UObject/PackageTrailer.h"
+#include "VirtualizationUtilities.h"
+
+FPackageDataValidator::FPackageDataValidator(const FName& InNameId, const FSubmitToolParameters& InParameters, TSharedRef<FSubmitToolServiceProvider> InServiceProvider, const FString& InDefinition)
+	: FValidatorBase(InNameId, InParameters, InServiceProvider, InDefinition)
+{
+	ParseDefinition(InDefinition);
+}
+
+void FPackageDataValidator::ParseDefinition(const FString& InDefinition)
+{
+	FStringOutputDevice Errors;
+	
+	FPackageDataValidatorDefinition* DefinitionToLoad = new FPackageDataValidatorDefinition;
+	FPackageDataValidatorDefinition::StaticStruct()->ImportText(*InDefinition, DefinitionToLoad, nullptr, 0, &Errors, FPackageDataValidatorDefinition::StaticStruct()->GetName());
+	
+	Definition.Reset(DefinitionToLoad);
+
+	if (!Errors.IsEmpty())
+	{
+		UE_LOG(LogSubmitTool, Error, TEXT("[%s] Error loading parameter file %s"), *GetValidatorName(), *Errors);
+		FModelInterface::SetErrorState();
+	}
+}
+
+bool FPackageDataValidator::Validate(const FString& InCLDescription, const TArray<FSourceControlStateRef>& InFilteredFilesInCL, const TArray<const FTag*>& InTags)
+{
+	if (CurrentJobs.IsEmpty() == false) // Could consider changing this to an assert when we are confident that it won't trigger
+	{
+		LogFailure(FString::Printf(TEXT("[%s] attempting to run a new validation while a previous validation pass is still running!"), *GetValidatorName()));
+		return false;
+	}
+
+	CheckForVirtualizedContent(InFilteredFilesInCL); // Launches a single job
+
+	// If we have no running work then the validation is complete
+	if (CurrentJobs.IsEmpty())
+	{
+		ValidationFinished(true);
+	}
+
+	return true;
+}
+
+void FPackageDataValidator::CheckForVirtualizedContent(const TArray<FSourceControlStateRef>& InFilteredFilesInCL)
+{
+	TUniquePtr<FJob> Job = MakeUnique<FJob>();
+
+	const FPackageDataValidatorDefinition* TypedDefinition = GetTypedDefinition<FPackageDataValidatorDefinition>();
+	check(TypedDefinition != nullptr);
+
+	// We could consider launching a task per entry in ExcludedExtensions if this validator ever becomes a bottleneck
+	Job->Task = UE::Tasks::Launch(UE_SOURCE_LOCATION, [this, InFilteredFilesInCL, ExcludedExtensions = TypedDefinition->ExcludedExtensions, &CancelToken = Job->CancelToken]() -> bool
+		{
+			bool bValid = true;
+
+			for (const FSourceControlStateRef& FileState : InFilteredFilesInCL)
+			{
+				for(const FString Extension : ExcludedExtensions)
+				{
+					if (CancelToken.IsCanceled())
+					{
+						break;
+					}
+
+					FString ProjectFilePath;
+					FString PluginFilePath;
+
+					if (UE::Virtualization::Utils::TryFindProject(FileState->GetFilename(), Extension, ProjectFilePath, PluginFilePath))
+					{
+						// Note we cannot use the asset registry in the SubmitTool as without the proper mount points we can
+						// only work with absolute file paths so we need to load the package trailer to tell if the package
+						// has virtualized data or not.
+						UE::FPackageTrailer Trailer;
+						if (UE::FPackageTrailer::TryLoadFromFile(FileState->GetFilename(), Trailer))
+						{
+							if (Trailer.GetNumPayloads(UE::EPayloadStorageType::Virtualized) > 0)
+							{
+								LogFailure(FString::Printf(TEXT("[%s] %s has virtualized content and will not work for a '%s' project, please rehydrate!"), *GetValidatorName(), *FileState->GetFilename(), *Extension));
+								bValid = false;
+							}
+						}
+					}
+				}
+			}
+
+			return bValid;
+		});
+
+	CurrentJobs.Add(MoveTemp(Job));
+}
+
+void FPackageDataValidator::Tick(float InDeltaTime)
+{
+	FValidatorBase::Tick(InDeltaTime);
+
+	if (HasWorkPending())
+	{
+		return;
+	}
+
+	bool bAllValidationsSuccessful = true;
+
+	for (const TUniquePtr<FJob>& Job : CurrentJobs)
+	{
+		if (Job && Job->Task.GetResult() == false)
+		{
+			
+			bAllValidationsSuccessful = false;
+		}
+	}
+
+	ValidationFinished(bAllValidationsSuccessful);
+
+	CurrentJobs.Empty();
+}
+
+void FPackageDataValidator::StopInternalValidations()
+{
+	// First, signal that the validation has been canceled to all running jobs.
+	for (const TUniquePtr<FJob>& Job : CurrentJobs)
+	{
+		if (Job)
+		{
+			Job->CancelToken.Cancel();
+		}
+	}
+
+	// Secondly, wait on any job that is still running until their cancellation is complete.
+	for (const TUniquePtr<FJob>& Job : CurrentJobs)
+	{
+		if (Job)
+		{
+			Job->Task.Wait();
+		}
+	}
+
+	// Lastly we can clean up the jobs.
+	CurrentJobs.Empty();
+}
+
+bool FPackageDataValidator::HasWorkPending() const
+{
+	for (const TUniquePtr<FJob>& Job : CurrentJobs)
+	{
+		if (Job && Job->Task.IsCompleted() == false)
+		{
+			return true;
+		}
+	}
+
+	return false;
+}
+
diff --git a/Engine/Source/Programs/SubmitTool/Private/Logic/Validators/PackageDataValidator.h b/Engine/Source/Programs/SubmitTool/Private/Logic/Validators/PackageDataValidator.h
new file mode 100644
--- /dev/null
+++ b/Engine/Source/Programs/SubmitTool/Private/Logic/Validators/PackageDataValidator.h
@@ -0,0 +1,42 @@
+// Copyright Epic Games, Inc. All Rights Reserved.
+
+#pragma once
+
+#include "ValidatorBase.h"
+
+#include "ISourceControlChangelistState.h"
+#include "Tasks/Task.h"
+
+class FPackageDataValidator final : public FValidatorBase
+{
+public:
+	using FValidatorBase::FValidatorBase;
+
+	FPackageDataValidator(const FName& InNameId, const FSubmitToolParameters& InParameters, TSharedRef<FSubmitToolServiceProvider> InServiceProvider, const FString& InDefinition);
+	virtual ~FPackageDataValidator() = default;
+
+private:
+
+	virtual void ParseDefinition(const FString& InDefinition) override;
+
+	virtual bool Validate(const FString& InCLDescription, const TArray<FSourceControlStateRef>& InFilteredFilesInCL, const TArray<const FTag*>& InTags) override;
+	virtual void Tick(float InDeltaTime) override;
+	virtual void StopInternalValidations() override;
+
+	virtual const FString& GetValidatorTypeName() const override
+	{
+		return SubmitToolParseConstants::PackageDataValidator;
+	}
+
+	bool HasWorkPending() const;
+
+	void CheckForVirtualizedContent(const TArray<FSourceControlStateRef>& InFilteredFilesInCL);
+
+	struct FJob
+	{
+		UE::Tasks::TTask<bool> Task;
+		UE::Tasks::FCancellationToken CancelToken;
+	};
+
+	TArray<TUniquePtr<FJob>> CurrentJobs;
+};
diff --git a/Engine/Source/Programs/SubmitTool/Private/Logic/Validators/ValidatorBase.h b/Engine/Source/Programs/SubmitTool/Private/Logic/Validators/ValidatorBase.h
--- a/Engine/Source/Programs/SubmitTool/Private/Logic/Validators/ValidatorBase.h
+++ b/Engine/Source/Programs/SubmitTool/Private/Logic/Validators/ValidatorBase.h
@@ -28,6 +28,7 @@ namespace SubmitToolParseConstants
 	const FString CustomValidator = TEXT("CustomValidator");
 	const FString CrossChangelistValidator = TEXT("CrossChangelistValidator");
 	const FString PreflightValidator = TEXT("PreflightValidator");
+	const FString PackageDataValidator = TEXT("PackageDataValidator");
 	const FString InvalidateNode = TEXT("Invalidated");
 	const FString VirtualizationToolOp = TEXT("VirtualizationTool");
 };
diff --git a/Engine/Source/Programs/SubmitTool/Private/Logic/Validators/ValidatorDefinition.h b/Engine/Source/Programs/SubmitTool/Private/Logic/Validators/ValidatorDefinition.h
--- a/Engine/Source/Programs/SubmitTool/Private/Logic/Validators/ValidatorDefinition.h
+++ b/Engine/Source/Programs/SubmitTool/Private/Logic/Validators/ValidatorDefinition.h
@@ -291,6 +291,15 @@ struct FUBTValidatorDefinition : public FValidatorRunExecutableDefinition
 	bool bUseStaticAnalyser = false;
 };
 
+USTRUCT()
+struct FPackageDataValidatorDefinition : public FValidatorDefinition
+{
+	GENERATED_BODY()
+
+	UPROPERTY()
+	TArray<FString> ExcludedExtensions;
+};
+
 USTRUCT()
 struct FVirtualizationToolDefinition : public FValidatorRunExecutableDefinition
 {
diff --git a/Engine/Source/Programs/SubmitTool/Private/Models/ModelInterface.cpp b/Engine/Source/Programs/SubmitTool/Private/Models/ModelInterface.cpp
--- a/Engine/Source/Programs/SubmitTool/Private/Models/ModelInterface.cpp
+++ b/Engine/Source/Programs/SubmitTool/Private/Models/ModelInterface.cpp
@@ -16,6 +16,7 @@
 #include "Logic/Validators/ValidatorBase.h"
 #include "Logic/Validators/ValidatorRunExecutable.h"
 #include "Logic/Validators/CrossChangelistValidator.h"
+#include "Logic/Validators/PackageDataValidator.h"
 #include "Logic/PreSubmitOperations/VirtualizationOperation.h"
 #include "Logic/Services/CacheDataService.h"
 #include "Logic/Services/SourceControl/SubmitToolPerforce.h"
@@ -220,6 +221,10 @@ void FModelInterface::ParseValidators() const
 		{
 			Tasks.Add(MakeShared<FPreflightValidator>(DefinitionPair.Key, Parameters, ServiceProvider.ToSharedRef(), DefinitionPair.Value));
 		}
+		else if (TaskDefinition.Type.Equals(SubmitToolParseConstants::PackageDataValidator, ESearchCase::IgnoreCase))
+		{
+			Tasks.Add(MakeShared<FPackageDataValidator>(DefinitionPair.Key, Parameters, ServiceProvider.ToSharedRef(), DefinitionPair.Value));
+		}
 		else
 		{
 			UE_LOG(LogSubmitTool, Error, TEXT("[%s] is not a recognized validator type and has not been activated."), *DefinitionPair.Key.ToString());
diff --git a/Engine/Source/Programs/SubmitTool/SubmitTool.Build.cs b/Engine/Source/Programs/SubmitTool/SubmitTool.Build.cs
--- a/Engine/Source/Programs/SubmitTool/SubmitTool.Build.cs
+++ b/Engine/Source/Programs/SubmitTool/SubmitTool.Build.cs
@@ -29,7 +29,8 @@ public class SubmitTool : ModuleRules
 				"JsonUtilities",
 				"Analytics",
 				"AnalyticsET",
-				"PakFile"
+				"PakFile",
+				"Virtualization"
 			}
 		);
 
