[SubmitTool] Use same substitution of $(StreamRoot) for AdditionalTasks as was added for the PF template in CL 37228717.

This re-enables additional tasks such as CompileShaders for preflights that include shader files.

#rnx

diff --git a/Engine/Source/Programs/SubmitTool/Private/Logic/PreflightService.cpp b/Engine/Source/Programs/SubmitTool/Private/Logic/PreflightService.cpp
--- a/Engine/Source/Programs/SubmitTool/Private/Logic/PreflightService.cpp
+++ b/Engine/Source/Programs/SubmitTool/Private/Logic/PreflightService.cpp
@@ -654,7 +654,7 @@ FString FPreflightService::GetAdditionalTasksString(const FPreflightTemplateDefi
 	const TArray<FString>& FilesInCl = bCheckShelveInstead ? ChangelistService->GetShelvedFilesDepotPaths() : ChangelistService->GetFilesDepotPaths();
 	for(const FPreflightAdditionalTask& AdditionalTask : InTemplate.AdditionalTasks)
 	{
-		FString RegexPat = FConfiguration::Substitute(AdditionalTask.RegexPath);
+		FString RegexPat = AdditionalTask.RegexPath.Replace(TEXT("$(StreamRoot)"), *StreamName, ESearchCase::IgnoreCase);
 		FRegexPattern Pattern = FRegexPattern(RegexPat, ERegexPatternFlags::CaseInsensitive);
 		for(const FString& File : FilesInCl)
 		{
