Forcing text outside of tags to be specified in a CL Description

diff --git a/Engine/Source/Programs/SubmitTool/Private/Logic/Validators/TagValidator.cpp b/Engine/Source/Programs/SubmitTool/Private/Logic/Validators/TagValidator.cpp
--- a/Engine/Source/Programs/SubmitTool/Private/Logic/Validators/TagValidator.cpp
+++ b/Engine/Source/Programs/SubmitTool/Private/Logic/Validators/TagValidator.cpp
@@ -19,6 +19,24 @@ bool FTagValidator::Validate(const FString& InCLDescription, const TArray<FSourc
 		bResult = false;
 	}
 
+	TArray<FString> DescriptionLines;
+	InCLDescription.ParseIntoArrayLines(DescriptionLines, true);
+	bool bHasDescription = false;
+	for (const FString& Line : DescriptionLines)
+	{
+		if (!Line.TrimStart().StartsWith(TEXT("#")))
+		{
+			bHasDescription = true;
+			break;
+		}
+	}
+
+	if (!bHasDescription)
+	{
+		LogFailure(FString::Printf(TEXT("[%s] Please introduce at least one description line that doesn't start with a tag"), *GetValidatorName()));
+		bResult = false;
+	}
+
 	// Wait for the GetUsersAndGroups asynchronously to finish validation
 	ServiceProvider.Pin()->GetService<ISTSourceControlService>()->GetUsersAndGroups(FOnUsersAndGroupsGet::FDelegate::CreateLambda([this, InTags](TArray<TSharedPtr<FUserData>>& InP4Users, TArray<TSharedPtr<FString>>& InP4Groups)
 		{
diff --git a/Engine/Source/Programs/SubmitTool/Private/Models/ModelInterface.cpp b/Engine/Source/Programs/SubmitTool/Private/Models/ModelInterface.cpp
--- a/Engine/Source/Programs/SubmitTool/Private/Models/ModelInterface.cpp
+++ b/Engine/Source/Programs/SubmitTool/Private/Models/ModelInterface.cpp
@@ -134,6 +134,7 @@ FModelInterface::FModelInterface(const FSubmitToolParameters& InParameters) :
 					{
 						UE_LOG(LogSubmitTool, Warning, TEXT("User has opted in to automatically submit on validation successful but not all validations succeeded. Fix optional validation errors or submit manually if you want to bypass them."));
 						FDialogFactory::ShowInformationDialog(FText::FromString(TEXT("Auto-Submit Cancelled")), FText::FromString(TEXT("Submit tool couldn't auto submit because there were optional validations that failed.\n\nFix these errors or manually submit if you are certain that you should ignore them.")));
+						bSubmitOnSuccessfulValidation = false;
 					}
 				}
 			}
