Validators are now categorized in Active/Inactive sections depending if they're applicable to the CL or not

diff --git a/Engine/Source/Programs/SubmitTool/Private/Logic/Validators/ValidatorBase.cpp b/Engine/Source/Programs/SubmitTool/Private/Logic/Validators/ValidatorBase.cpp
--- a/Engine/Source/Programs/SubmitTool/Private/Logic/Validators/ValidatorBase.cpp
+++ b/Engine/Source/Programs/SubmitTool/Private/Logic/Validators/ValidatorBase.cpp
@@ -171,6 +171,11 @@ const FString FValidatorBase::GetStatusText() const
 
 	if(State == EValidationStates::Skipped || State == EValidationStates::Not_Run)
 	{
+		if (State == EValidationStates::Skipped && !IsRelevantToCL())
+		{
+			return TEXT("Not Applicable");
+		}
+
 		return StateStr;
 	}
 
@@ -263,6 +268,15 @@ bool FValidatorBase::EvaluateTagSkip()
 	return false;
 }
 
+bool FValidatorBase::IsRelevantToCL() const
+{
+	TSharedPtr<FChangelistService> ChangelistService = ServiceProvider.Pin()->GetService<FChangelistService>();
+
+	TArray<FSourceControlStateRef> IncrementallySkippedFiles;
+	TArray<FSourceControlStateRef> OutFiles;
+	return AppliesToCL(ChangelistService->GetCLDescription(), ChangelistService->GetFilesInCL(), ServiceProvider.Pin()->GetService<FTagService>()->GetTagsArray(), OutFiles, IncrementallySkippedFiles, false);
+}
+
 void FValidatorBase::SetSelectedOption(const FString& InOptionName, const FString& InOptionValue)
 {
 	UE_LOG(LogValidators, Log, TEXT("[%s] Task stopped due to a change in options, %s = %s"), *GetValidatorName(), *InOptionName, *InOptionValue);
diff --git a/Engine/Source/Programs/SubmitTool/Private/Logic/Validators/ValidatorBase.h b/Engine/Source/Programs/SubmitTool/Private/Logic/Validators/ValidatorBase.h
--- a/Engine/Source/Programs/SubmitTool/Private/Logic/Validators/ValidatorBase.h
+++ b/Engine/Source/Programs/SubmitTool/Private/Logic/Validators/ValidatorBase.h
@@ -141,10 +141,17 @@ public:
 		State = EValidationStates::Queued;
 	}
 
+	void SetNotApplicable()
+	{
+		State = EValidationStates::Skipped;
+	}
+
 	virtual bool Activate();
 
 	virtual void InvalidateLocalFileModifications();
 	virtual bool EvaluateTagSkip();
+	virtual bool IsRelevantToCL() const;
+
 
 	const FString GetStatusText() const;
 	const EValidationStates GetValidatorState() const
diff --git a/Engine/Source/Programs/SubmitTool/Private/Logic/Validators/ValidatorDefinition.h b/Engine/Source/Programs/SubmitTool/Private/Logic/Validators/ValidatorDefinition.h
--- a/Engine/Source/Programs/SubmitTool/Private/Logic/Validators/ValidatorDefinition.h
+++ b/Engine/Source/Programs/SubmitTool/Private/Logic/Validators/ValidatorDefinition.h
@@ -107,12 +107,6 @@ struct FValidatorDefinition
 	/// 
 	UPROPERTY()
 	TArray<FName> ExecutionBlockGroups;
-	
-	///
-	///	Name of the UI Section this Validator is part of
-	/// 
-	UPROPERTY()
-	FName UIGroup;
 
 	///
 	///	If this Validator runs on files marked for delete
diff --git a/Engine/Source/Programs/SubmitTool/Private/View/Widgets/ValidatorsWidget.cpp b/Engine/Source/Programs/SubmitTool/Private/View/Widgets/ValidatorsWidget.cpp
--- a/Engine/Source/Programs/SubmitTool/Private/View/Widgets/ValidatorsWidget.cpp
+++ b/Engine/Source/Programs/SubmitTool/Private/View/Widgets/ValidatorsWidget.cpp
@@ -36,35 +36,68 @@ void SValidatorsWidget::Construct(const FArguments& InArgs)
 	};
 
 	ModelInterface = InArgs._ModelInterface;
-	TArray<TWeakPtr<const FValidatorBase>> ValidatorsArray = ModelInterface->GetValidators();
 	OnViewLog = InArgs._OnViewLog;
 
-	TSharedRef<SVerticalBox> VBox = SNew(SVerticalBox);
+	ChildSlot
+	[
+		BuildValidatorsView()
+	];
+
+	OnFilesRefreshed = ModelInterface->FileRefreshedCallback.AddLambda([this] { RefreshValidatorView(); });
+	OnPrepareSubmit = ModelInterface->PrepareSubmitCallBack.AddLambda([this] { RefreshValidatorView(true); });
+}
+
+SValidatorsWidget::~SValidatorsWidget()
+{
+	ModelInterface->FileRefreshedCallback.Remove(OnFilesRefreshed);
+	ModelInterface->PrepareSubmitCallBack.Remove(OnPrepareSubmit);
+}
 
+TSharedRef<SWidget> SValidatorsWidget::BuildValidatorsView(bool bListPreSubmitOperations)
+{
+	Validators.Reset();
+
+	TArray<TWeakPtr<const FValidatorBase>> ValidatorsArray = bListPreSubmitOperations ? ModelInterface->GetPreSubmitOperations() : ModelInterface->GetValidators();
+	const FName ActiveSection = bListPreSubmitOperations ? FName(TEXT("Pre-Submit Operations")) : FName(TEXT("Active Validations"));
+	const FName InactiveSection = bListPreSubmitOperations ? FName(TEXT("Inactive Pre-Submit Operations (Not Applicable to your CL)")) : FName(TEXT("Inactive Validators (Not Applicable to your CL)"));
+
+	TSharedRef<SVerticalBox> VBox = SNew(SVerticalBox);
+	
 	for(const TWeakPtr<const FValidatorBase>& Validator : ValidatorsArray)
 	{
-		Validators.FindOrAdd(Validator.Pin()->Definition->UIGroup).Add(Validator);
-	}
+		FName DisplayName;
 
-	for(const TPair<FName, TArray<TWeakPtr<const FValidatorBase>>>& Pair : Validators)
-	{
-		FName DisplayName = Pair.Key;
-		if(Pair.Key.IsNone())
+		if (!Validator.Pin()->IsRelevantToCL())
+		{
+			DisplayName = InactiveSection;
+		}
+		else
 		{
-			DisplayName = FName(TEXT("General Validations"));
+			DisplayName = ActiveSection;
 		}
 
+		Validators.FindOrAdd(DisplayName).Add(Validator);
+	}
+
+	for(const TPair<FName, TArray<TWeakPtr<const FValidatorBase>>>& Pair : Validators)
+	{
 		VBox->AddSlot()
 			.Padding(FMargin(0.f, 2.f))
 			.AutoHeight()
 			.AttachWidget
 			(
 				SNew(SExpandableArea)
-				.InitiallyCollapsed(!FSubmitToolUserPrefs::Get()->UISectionExpandState.FindOrAdd(DisplayName, false))
+				.InitiallyCollapsed(!bListPreSubmitOperations && !FSubmitToolUserPrefs::Get()->UISectionExpandState.FindOrAdd(Pair.Key, false))
 				.BorderImage(FAppStyle::Get().GetBrush("Brushes.Header"))
 				.BodyBorderImage(FAppStyle::Get().GetBrush("Brushes.Recessed"))
-				.OnAreaExpansionChanged_Lambda([DisplayName](bool bExpanded) { FSubmitToolUserPrefs::Get()->UISectionExpandState.FindOrAdd(DisplayName, false) = bExpanded; })
-				.BorderBackgroundColor_Lambda([&ValidatorArray = Pair.Value] { 
+				.OnAreaExpansionChanged_Lambda([DisplayName = Pair.Key](bool bExpanded) { FSubmitToolUserPrefs::Get()->UISectionExpandState.FindOrAdd(DisplayName, false) = bExpanded; })
+				.BorderBackgroundColor_Lambda([&ValidatorArray = Pair.Value, DisplayName = Pair.Key, InactiveSection]
+				{
+					if (DisplayName.IsEqual(InactiveSection))
+					{
+						return FAppStyle::GetColor("ValidatorStateSuccess");
+					}
+
 					bool bRunning = false;
 					bool bValid = true;
 					for(const TWeakPtr<const FValidatorBase>& Validator : ValidatorArray)
@@ -96,7 +129,7 @@ void SValidatorsWidget::Construct(const FArguments& InArgs)
 				.HeaderContent()
 				[
 					SNew(STextBlock)
-					.Text_Lambda([DisplayName] {return FText::FromName(DisplayName); })
+					.Text_Lambda([DisplayName = Pair.Key] {return FText::FromName(DisplayName); })
 				]
 				.BodyContent()
 				[
@@ -111,14 +144,16 @@ void SValidatorsWidget::Construct(const FArguments& InArgs)
 							.IsFocusable(false)
 					]
 				]
-			);
-			
+			);			
 	}
 
-	ChildSlot
-	[
-		VBox
-	];
+	return VBox;
+}
+
+void SValidatorsWidget::RefreshValidatorView(bool bListPreSubmitOperations)
+{
+	ChildSlot.DetachWidget();
+	ChildSlot.AttachWidget(BuildValidatorsView(bListPreSubmitOperations));
 }
 
 TSharedRef<SHeaderRow> SValidatorsWidget::ConstructHeadersRow(const FName& GroupName)
@@ -288,7 +323,7 @@ public:
 					{
 					case EValidationStates::Failed:
 					case EValidationStates::Timeout:
-						return FAppStyle::GetColor("ValidatorStateFail");
+						return Node.Pin()->Definition->IsRequired ? FAppStyle::GetColor("ValidatorStateFail") : FAppStyle::GetColor("ValidatorStateWarning");
 						break;
 					case EValidationStates::Valid:
 					case EValidationStates::Skipped:
@@ -311,7 +346,7 @@ public:
 				.Padding(3,0,0,0)
 				[
 					SNew(SButton)
-					.IsEnabled_Lambda([this]() {return ModelInterface->GetInputEnabled(); })
+					.IsEnabled_Lambda([this]() {return ModelInterface->GetInputEnabled() && Node.Pin()->IsRelevantToCL(); })
 					.HAlign(EHorizontalAlignment::HAlign_Center)
 					.VAlign(EVerticalAlignment::VAlign_Center)
 					.OnClicked_Lambda([this]()
@@ -331,7 +366,7 @@ public:
 				.Padding(3,0,0,0)
 				[
 					SNew(SButton)
-					.IsEnabled_Lambda([this]() {return ModelInterface->GetInputEnabled(); })
+					.IsEnabled_Lambda([this]() {return ModelInterface->GetInputEnabled() && Node.Pin()->IsRelevantToCL(); })
 					.HAlign(EHorizontalAlignment::HAlign_Center)
 					.VAlign(EVerticalAlignment::VAlign_Center)
 					.OnClicked_Lambda([this]()
diff --git a/Engine/Source/Programs/SubmitTool/Private/View/Widgets/ValidatorsWidget.h b/Engine/Source/Programs/SubmitTool/Private/View/Widgets/ValidatorsWidget.h
--- a/Engine/Source/Programs/SubmitTool/Private/View/Widgets/ValidatorsWidget.h
+++ b/Engine/Source/Programs/SubmitTool/Private/View/Widgets/ValidatorsWidget.h
@@ -36,11 +36,17 @@ public:
 	SLATE_END_ARGS()
 	
 	void Construct(const FArguments& InArgs);
+	virtual ~SValidatorsWidget();
 
 private:
+
+	TSharedRef<SWidget> BuildValidatorsView(bool bListPreSubmitOperations = false);
+	void RefreshValidatorView(bool bListPreSubmitOperations = false);
+
 	FModelInterface* ModelInterface;
 	
 	TSharedPtr<SListView<TWeakPtr<const FValidatorBase>>> ValidatorsListView;
+	TSharedPtr<SListView<TWeakPtr<const FValidatorBase>>> PreSubmitListView;
 	
 	TMap<FName, TArray<TWeakPtr<const FValidatorBase>>> Validators;
 	TArray<FValidatorColumn> Columns;
@@ -54,4 +60,6 @@ private:
 	FName SortByColumn;
     EColumnSortMode::Type SortMode;
 	FOnViewValidatorLog OnViewLog;
+	FDelegateHandle OnFilesRefreshed;
+	FDelegateHandle OnPrepareSubmit;
 };
