Adding window maximization to the submit tool user prefs
Added pre-submit operations view when submitting

diff --git a/Engine/Source/Programs/SubmitTool/Private/Logic/CredentialsService.cpp b/Engine/Source/Programs/SubmitTool/Private/Logic/CredentialsService.cpp
--- a/Engine/Source/Programs/SubmitTool/Private/Logic/CredentialsService.cpp
+++ b/Engine/Source/Programs/SubmitTool/Private/Logic/CredentialsService.cpp
@@ -135,38 +135,42 @@ const FString FCredentialsService::GetCredentialsFilepath() const
 void FCredentialsService::GetOIDCToken()
 {	
 	UE_LOG(LogSubmitTool, Log, TEXT("Obtaining new OIDCToken"));
-	GetOIDCTask = UE::Tasks::Launch(UE_SOURCE_LOCATION,
-		[this] {
-			FOnOutputLine OutputLineProcess = FOnOutputLine::CreateLambda([this](const FString& OutputLine, const EProcessOutputType& OutputType) {
-				if(OutputType == EProcessOutputType::ProcessError)
-				{
-					UE_LOG(LogSubmitTool, Error, TEXT("%s"), *OutputLine);
-				}
-				else
-				{
-					// OIDCToken console output token starts with a {
-					if(OutputLine.StartsWith(TEXT("{")))
+
+	if (!GetOIDCTask.IsValid() || GetOIDCTask.IsCompleted())
+	{
+		GetOIDCTask = UE::Tasks::Launch(UE_SOURCE_LOCATION,
+			[this] {
+				FOnOutputLine OutputLineProcess = FOnOutputLine::CreateLambda([this](const FString& OutputLine, const EProcessOutputType& OutputType) {
+					if(OutputType == EProcessOutputType::ProcessError)
 					{
-						ParseOIDCTokenData(OutputLine);
+						UE_LOG(LogSubmitTool, Error, TEXT("%s"), *OutputLine);
 					}
 					else
 					{
-						UE_LOG(LogSubmitToolDebug, Log, TEXT("%s"), *OutputLine);
+						// OIDCToken console output token starts with a {
+						if(OutputLine.StartsWith(TEXT("{")))
+						{
+							ParseOIDCTokenData(OutputLine);
+						}
+						else
+						{
+							UE_LOG(LogSubmitToolDebug, Log, TEXT("%s"), *OutputLine);
+						}
 					}
-				}
-				});
+					});
 
-			FProcessWrapper GetToken = FProcessWrapper(TEXT("Oidc"), Parameters.OAuthTokenTool, FString::Printf(TEXT("%s"), *Parameters.OAuthArgs), FOnCompleted(), OutputLineProcess);
-			GetToken.Start(true);
+				FProcessWrapper GetToken = FProcessWrapper(TEXT("Oidc"), Parameters.OAuthTokenTool, FString::Printf(TEXT("%s"), *Parameters.OAuthArgs), FOnCompleted(), OutputLineProcess);
+				GetToken.Start(true);
 
-			if(GetToken.ExitCode != 0)
-			{
-				UE_LOG(LogSubmitTool, Warning, TEXT("Couldn't obtain OIDC credentials"));
-				return false;
-			}
+				if(GetToken.ExitCode != 0)
+				{
+					UE_LOG(LogSubmitTool, Warning, TEXT("Couldn't obtain OIDC credentials"));
+					return false;
+				}
 
-			return true;
-		});
+				return true;
+			});
+	}
 }
 
 bool FCredentialsService::ParseOIDCTokenData(const FString& InToken)
diff --git a/Engine/Source/Programs/SubmitTool/Private/Logic/TasksService.cpp b/Engine/Source/Programs/SubmitTool/Private/Logic/TasksService.cpp
--- a/Engine/Source/Programs/SubmitTool/Private/Logic/TasksService.cpp
+++ b/Engine/Source/Programs/SubmitTool/Private/Logic/TasksService.cpp
@@ -125,6 +125,12 @@ bool FTasksService::QueueForExecution(const TSharedPtr<FValidatorBase>& InTask,
 		return false;
 	}
 
+	if ((!InTask->IsRelevantToCL() && !InbForceRun))
+	{
+		InTask->SetNotApplicable();
+		return false;
+	}
+
 	if(!InTask->GetHasPassed() || InbForceRun)
 	{
 		if(InTask->Definition->DependsOn.Num() > 0)
diff --git a/Engine/Source/Programs/SubmitTool/Private/Models/ModelInterface.cpp b/Engine/Source/Programs/SubmitTool/Private/Models/ModelInterface.cpp
--- a/Engine/Source/Programs/SubmitTool/Private/Models/ModelInterface.cpp
+++ b/Engine/Source/Programs/SubmitTool/Private/Models/ModelInterface.cpp
@@ -23,6 +23,7 @@
 #include "Logic/Services/SourceControl/SubmitToolPerforce.h"
 
 ESubmitToolAppState FModelInterface::SubmitToolState = ESubmitToolAppState::Initializing;
+FOnStateChanged FModelInterface::OnStateChanged;
 
 FModelInterface::FModelInterface(const FSubmitToolParameters& InParameters) :
 	Parameters(InParameters)
@@ -823,10 +824,12 @@ bool FModelInterface::Tick(float InDeltaTime)
 
 void FModelInterface::ChangeState(ESubmitToolAppState newState, bool bForce)
 {
+	const ESubmitToolAppState CurrentState = SubmitToolState;
 	if(bForce)
 	{
 		UE_LOG(LogSubmitToolDebug, Log, TEXT("Transitioned state from '%s' to '%s'"), *StaticEnum<ESubmitToolAppState>()->GetNameStringByValue(static_cast<int64>(SubmitToolState)), *StaticEnum<ESubmitToolAppState>()->GetNameStringByValue(static_cast<int64>(newState)));
 		SubmitToolState = newState;
+		OnStateChanged.Broadcast(CurrentState, SubmitToolState);
 	}
 	else
 	{
@@ -837,6 +840,7 @@ void FModelInterface::ChangeState(ESubmitToolAppState newState, bool bForce)
 			{
 				UE_LOG(LogSubmitToolDebug, Log, TEXT("Transitioned state from '%s' to '%s'"), *StaticEnum<ESubmitToolAppState>()->GetNameStringByValue(static_cast<int64>(SubmitToolState)), *StaticEnum<ESubmitToolAppState>()->GetNameStringByValue(static_cast<int64>(newState)));
 				SubmitToolState = newState;
+				OnStateChanged.Broadcast(CurrentState, SubmitToolState);
 			}
 			else
 			{
diff --git a/Engine/Source/Programs/SubmitTool/Private/Models/ModelInterface.h b/Engine/Source/Programs/SubmitTool/Private/Models/ModelInterface.h
--- a/Engine/Source/Programs/SubmitTool/Private/Models/ModelInterface.h
+++ b/Engine/Source/Programs/SubmitTool/Private/Models/ModelInterface.h
@@ -84,6 +84,7 @@ namespace SubmitToolAppState {
 
 DECLARE_MULTICAST_DELEGATE(FPreSubmitCallBack)
 DECLARE_MULTICAST_DELEGATE(FFilesRefresh)
+DECLARE_MULTICAST_DELEGATE_TwoParams(FOnStateChanged, const ESubmitToolAppState /*InFrom*/, const ESubmitToolAppState /*InTo*/)
 
 class FModelInterface
 {
@@ -135,6 +136,7 @@ public:
 	bool HasSubmitToolTag() const;
 	bool IsValidationRunning() const												{ return ValidationService->GetIsAnyTaskRunning(); }
 	const TArray<TWeakPtr<const FValidatorBase>>& GetValidators() const				{ return ValidationService->GetTasks(); }
+	const TArray<TWeakPtr<const FValidatorBase>>& GetPreSubmitOperations() const	{ return PresubmitOperationsService->GetTasks(); }
 	void CancelValidations(const FName& InValidatorId = FName()) const				{ ValidationService->StopTasks(InValidatorId); }
 	void CheckForFileEdits() const													{ ValidationService->CheckForLocalFileEdit(); }
 
@@ -203,6 +205,7 @@ public:
 	bool ValidateIntegrationOptions(bool bSilent) const	{ return FNIntegrationService->ValidateIntegrationOptions(bSilent); }
 
 
+	static FOnStateChanged OnStateChanged;
 	FPreSubmitCallBack PrepareSubmitCallBack;
 	FFilesRefresh FileRefreshedCallback;
 	bool bSubmitOnSuccessfulValidation = false;
diff --git a/Engine/Source/Programs/SubmitTool/Private/Models/SubmitToolUserPrefs.h b/Engine/Source/Programs/SubmitTool/Private/Models/SubmitToolUserPrefs.h
--- a/Engine/Source/Programs/SubmitTool/Private/Models/SubmitToolUserPrefs.h
+++ b/Engine/Source/Programs/SubmitTool/Private/Models/SubmitToolUserPrefs.h
@@ -20,6 +20,9 @@ struct FSubmitToolUserPrefs
 	
 	UPROPERTY()
 	bool bExpandFilesInCL = false;
+	
+	UPROPERTY()
+	bool bWindowMaximized = false;
 
 	UPROPERTY()
 	FVector2D WindowPosition = FVector2D();
diff --git a/Engine/Source/Programs/SubmitTool/Private/SubmitToolApp.cpp b/Engine/Source/Programs/SubmitTool/Private/SubmitToolApp.cpp
--- a/Engine/Source/Programs/SubmitTool/Private/SubmitToolApp.cpp
+++ b/Engine/Source/Programs/SubmitTool/Private/SubmitToolApp.cpp
@@ -151,6 +151,12 @@ int RunSubmitTool(const TCHAR* CommandLine, const FGuid& InSessionID)
 				Tab->GetParentWindow()->Resize(FDeprecateSlateVector2D(1024, 768));
 			}
 
+
+			if (UserPrefs->bWindowMaximized)
+			{
+				Tab->GetParentWindow()->Maximize();
+			}
+
 			// Loop the engine
 			while(!IsEngineExitRequested())
 			{
diff --git a/Engine/Source/Programs/SubmitTool/Private/View/SubmitToolWindow.cpp b/Engine/Source/Programs/SubmitTool/Private/View/SubmitToolWindow.cpp
--- a/Engine/Source/Programs/SubmitTool/Private/View/SubmitToolWindow.cpp
+++ b/Engine/Source/Programs/SubmitTool/Private/View/SubmitToolWindow.cpp
@@ -27,6 +27,7 @@ TSharedRef<SDockTab> SubmitToolWindow::BuildMainTab(TSharedPtr<SWindow> InParent
 							{
 								FSubmitToolUserPrefs::Get()->WindowPosition = Tab->GetParentWindow()->GetPositionInScreen();
 								FSubmitToolUserPrefs::Get()->WindowSize = Tab->GetParentWindow()->GetClientSizeInScreen();
+								FSubmitToolUserPrefs::Get()->bWindowMaximized = Tab->GetParentWindow()->IsWindowMaximized();
 							}
 						})
 		.Label_Lambda([&ModelInterface = ModelInterface]() {return FText::FromString(FString::Printf(TEXT("Changelist: %s"), *ModelInterface->GetCLID())); });
diff --git a/Engine/Source/Programs/SubmitTool/Private/View/Widgets/ValidatorsWidget.cpp b/Engine/Source/Programs/SubmitTool/Private/View/Widgets/ValidatorsWidget.cpp
--- a/Engine/Source/Programs/SubmitTool/Private/View/Widgets/ValidatorsWidget.cpp
+++ b/Engine/Source/Programs/SubmitTool/Private/View/Widgets/ValidatorsWidget.cpp
@@ -44,13 +44,19 @@ void SValidatorsWidget::Construct(const FArguments& InArgs)
 	];
 
 	OnFilesRefreshed = ModelInterface->FileRefreshedCallback.AddLambda([this] { RefreshValidatorView(); });
-	OnPrepareSubmit = ModelInterface->PrepareSubmitCallBack.AddLambda([this] { RefreshValidatorView(true); });
+	OnPrepareSubmit = FModelInterface::OnStateChanged.AddLambda([this] (const ESubmitToolAppState InFrom, const ESubmitToolAppState InTo)
+		{ 
+			if (InTo == ESubmitToolAppState::Submitting || InFrom == ESubmitToolAppState::Submitting)
+			{
+				RefreshValidatorView(InTo == ESubmitToolAppState::Submitting);
+			}
+		});
 }
 
 SValidatorsWidget::~SValidatorsWidget()
 {
 	ModelInterface->FileRefreshedCallback.Remove(OnFilesRefreshed);
-	ModelInterface->PrepareSubmitCallBack.Remove(OnPrepareSubmit);
+	FModelInterface::OnStateChanged.Remove(OnPrepareSubmit);
 }
 
 TSharedRef<SWidget> SValidatorsWidget::BuildValidatorsView(bool bListPreSubmitOperations)
