Including CL in submit tool version through preprocessor definitions
Compilation fix
Version bump

diff --git a/Engine/Source/Programs/SubmitTool/Private/Logic/Validators/PackageDataValidator.cpp b/Engine/Source/Programs/SubmitTool/Private/Logic/Validators/PackageDataValidator.cpp
--- a/Engine/Source/Programs/SubmitTool/Private/Logic/Validators/PackageDataValidator.cpp
+++ b/Engine/Source/Programs/SubmitTool/Private/Logic/Validators/PackageDataValidator.cpp
@@ -61,7 +61,7 @@ void FPackageDataValidator::CheckForVirtualizedContent(const TArray<FSourceContr
 
 			for (const FSourceControlStateRef& FileState : InFilteredFilesInCL)
 			{
-				for(const FString Extension : ExcludedExtensions)
+				for(const FString& Extension : ExcludedExtensions)
 				{
 					if (CancelToken.IsCanceled())
 					{
diff --git a/Engine/Source/Programs/SubmitTool/Private/Version/AppVersion.cpp b/Engine/Source/Programs/SubmitTool/Private/Version/AppVersion.cpp
--- a/Engine/Source/Programs/SubmitTool/Private/Version/AppVersion.cpp
+++ b/Engine/Source/Programs/SubmitTool/Private/Version/AppVersion.cpp
@@ -38,7 +38,7 @@ FString FAppVersion::GetVersion()
 			{ TEXT("Application"), SUBMIT_TOOL_APPNAME },
 			{ TEXT("ApplicatonVersion"), SUBMIT_TOOL_VERSION_STRING },
 			{ TEXT("BuildType"), BuildType },
-			{ TEXT("BuildId"), GetBuildId() }
+			{ TEXT("BuildId"), SUBMIT_TOOL_CHANGELIST_STRING }
 		};
 
 		Version = FString::Format(TEXT("{Application}-{BuildType}-{ApplicatonVersion}.{BuildId}"), FormatArgs);
diff --git a/Engine/Source/Programs/SubmitTool/Private/Version/AppVersionDefines.h b/Engine/Source/Programs/SubmitTool/Private/Version/AppVersionDefines.h
--- a/Engine/Source/Programs/SubmitTool/Private/Version/AppVersionDefines.h
+++ b/Engine/Source/Programs/SubmitTool/Private/Version/AppVersionDefines.h
@@ -20,8 +20,16 @@
 #endif
 
 #define SUBMIT_TOOL_VERSION_MAJOR 5
-#define SUBMIT_TOOL_VERSION_MINOR 2
-#define SUBMIT_TOOL_VERSION_PATCH 2
+#define SUBMIT_TOOL_VERSION_MINOR 3
+#define SUBMIT_TOOL_VERSION_PATCH 0
+
+#ifndef SUBMIT_TOOL_CHANGELIST
+#define SUBMIT_TOOL_CHANGELIST_STRING TEXT("NoCL")
+#endif
+
+#ifdef SUBMIT_TOOL_CHANGELIST
+#define SUBMIT_TOOL_CHANGELIST_STRING TEXT("CL-") VERSION_STRINGIFY(SUBMIT_TOOL_CHANGELIST)
+#endif
 
 #define SUBMIT_TOOL_VERSION_STRING \
 	VERSION_STRINGIFY(SUBMIT_TOOL_VERSION_MAJOR)\
