Added a basic JSON Validator

diff --git a/Engine/Source/Programs/SubmitTool/Private/Logic/Validators/JsonValidator.cpp b/Engine/Source/Programs/SubmitTool/Private/Logic/Validators/JsonValidator.cpp
new file mode 100644
--- /dev/null
+++ b/Engine/Source/Programs/SubmitTool/Private/Logic/Validators/JsonValidator.cpp
@@ -0,0 +1,100 @@
+// Copyright Epic Games, Inc. All Rights Reserved.
+
+#include "JsonValidator.h"
+
+#include "Serialization/JsonReader.h"
+#include "Serialization/JsonSerializer.h"
+#include "Misc/FileHelper.h"
+#include "Misc/Paths.h"
+
+#include "Models/ModelInterface.h"
+
+FJsonValidator::FJsonValidator(const FName& InNameId, const FSubmitToolParameters& InParameters, TSharedRef<FSubmitToolServiceProvider> InServiceProvider, const FString& InDefinition) :
+	FValidatorBase(InNameId, InParameters, InServiceProvider, InDefinition)
+{
+}
+
+bool FJsonValidator::Validate(const FString& InCLDescription, const TArray<FSourceControlStateRef>& InFilteredFilesInCL, const TArray<const FTag*>& InTags)
+{
+	TArray<FString> JsonFilenames;
+	for (const FSourceControlStateRef& File : InFilteredFilesInCL)
+	{
+		JsonFilenames.Add(File->GetFilename());
+	}
+
+	CurrentJob = MakeUnique<FJob>();
+
+	CurrentJob->Task = UE::Tasks::Launch(UE_SOURCE_LOCATION, [this, JsonFilenames, &CancelToken = CurrentJob->CancelToken]() -> bool
+		{
+			bool bValid = true;
+
+			for (const FString& Filename : JsonFilenames)
+			{
+				if (CancelToken.IsCanceled())
+				{
+					break;
+				}
+
+				FString JSON;
+				if (!FFileHelper::LoadFileToString(JSON, *Filename))
+				{
+					LogFailure(FString::Printf(TEXT("[%s] %s could not be loaded"), *GetValidatorName(), *Filename));
+					bValid = false;
+					continue;
+				}
+
+				TSharedRef<TJsonReader<>> JsonReader = TJsonReaderFactory<>::Create(JSON);
+				TSharedPtr<FJsonObject> JsonObject;
+				if (!FJsonSerializer::Deserialize(JsonReader, JsonObject) || JsonObject.IsValid())
+				{
+					const FString& Error = JsonReader->GetErrorMessage();
+					LogFailure(FString::Printf(TEXT("[%s] %s is an invalid JSON file: %s"), *GetValidatorName(), *Filename, *Error));
+					bValid = false;
+				}
+			}
+
+			return bValid;
+		});
+
+	return true;
+}
+
+void FJsonValidator::Tick(float InDeltaTime)
+{
+	FValidatorBase::Tick(InDeltaTime);
+
+	if (HasWorkPending())
+	{
+		return;
+	}
+
+	const bool bAllValidationsSuccessful = !(CurrentJob && CurrentJob->Task.GetResult() == false);
+	ValidationFinished(bAllValidationsSuccessful);
+
+	CurrentJob.Reset();
+}
+
+void FJsonValidator::StopInternalValidations()
+{
+	if (CurrentJob)
+	{
+		// First, signal that the validation has been canceled to all running jobs.
+		CurrentJob->CancelToken.Cancel();
+
+		// Secondly, wait on any job that is still running until their cancellation is complete.
+		CurrentJob->Task.Wait();
+
+		// Lastly we can clean up the jobs.
+		CurrentJob.Reset();
+	}
+}
+
+bool FJsonValidator::HasWorkPending() const
+{
+	if (CurrentJob && CurrentJob->Task.IsCompleted() == false)
+	{
+		return true;
+	}
+	return false;
+}
+
diff --git a/Engine/Source/Programs/SubmitTool/Private/Logic/Validators/JsonValidator.h b/Engine/Source/Programs/SubmitTool/Private/Logic/Validators/JsonValidator.h
new file mode 100644
--- /dev/null
+++ b/Engine/Source/Programs/SubmitTool/Private/Logic/Validators/JsonValidator.h
@@ -0,0 +1,33 @@
+// Copyright Epic Games, Inc. All Rights Reserved.
+
+#pragma once
+
+#include "ValidatorBase.h"
+#include "Tasks/Task.h"
+
+class FJsonValidator final : public FValidatorBase
+{
+public:
+	FJsonValidator(const FName& InNameId, const FSubmitToolParameters& InParameters, TSharedRef<FSubmitToolServiceProvider> InServiceProvider, const FString& InDefinition);
+	virtual bool Validate(const FString& InCLDescription, const TArray<FSourceControlStateRef>& InFilteredFilesInCL, const TArray<const FTag*>& InTags) override;
+
+	virtual const FString& GetValidatorTypeName() const override
+	{
+		return SubmitToolParseConstants::JsonValidator;
+	}
+
+private:
+	virtual void Tick(float InDeltaTime) override;
+	virtual void StopInternalValidations() override;
+
+	bool HasWorkPending() const;
+
+	struct FJob
+	{
+		UE::Tasks::TTask<bool> Task;
+		UE::Tasks::FCancellationToken CancelToken;
+	};
+
+	TUniquePtr<FJob> CurrentJob;
+
+};
diff --git a/Engine/Source/Programs/SubmitTool/Private/Logic/Validators/ValidatorBase.h b/Engine/Source/Programs/SubmitTool/Private/Logic/Validators/ValidatorBase.h
--- a/Engine/Source/Programs/SubmitTool/Private/Logic/Validators/ValidatorBase.h
+++ b/Engine/Source/Programs/SubmitTool/Private/Logic/Validators/ValidatorBase.h
@@ -30,6 +30,7 @@ namespace SubmitToolParseConstants
 	const FString CrossChangelistValidator = TEXT("CrossChangelistValidator");
 	const FString PreflightValidator = TEXT("PreflightValidator");
 	const FString PackageDataValidator = TEXT("PackageDataValidator");
+	const FString JsonValidator = TEXT("JsonValidator");
 	const FString InvalidateNode = TEXT("Invalidated");
 	const FString VirtualizationToolOp = TEXT("VirtualizationTool");
 };
diff --git a/Engine/Source/Programs/SubmitTool/Private/Models/ModelInterface.cpp b/Engine/Source/Programs/SubmitTool/Private/Models/ModelInterface.cpp
--- a/Engine/Source/Programs/SubmitTool/Private/Models/ModelInterface.cpp
+++ b/Engine/Source/Programs/SubmitTool/Private/Models/ModelInterface.cpp
@@ -18,6 +18,7 @@
 #include "Logic/Validators/CrossChangelistValidator.h"
 #include "Logic/Validators/PackageDataValidator.h"
 #include "Logic/Validators/ShaderValidator.h"
+#include "Logic/Validators/JsonValidator.h"
 #include "Logic/PreSubmitOperations/VirtualizationOperation.h"
 #include "Logic/Services/CacheDataService.h"
 #include "Logic/Services/SourceControl/SubmitToolPerforce.h"
@@ -232,12 +233,17 @@ void FModelInterface::ParseValidators() const
 		{
 			Tasks.Add(MakeShared<FShaderValidator>(DefinitionPair.Key, Parameters, ServiceProvider.ToSharedRef(), DefinitionPair.Value));
 		}
+		else if (TaskDefinition.Type.Equals(SubmitToolParseConstants::JsonValidator, ESearchCase::IgnoreCase))
+		{
+			Tasks.Add(MakeShared<FJsonValidator>(DefinitionPair.Key, Parameters, ServiceProvider.ToSharedRef(), DefinitionPair.Value));
+		}
 		else
 		{
 			UE_LOG(LogSubmitTool, Error, TEXT("[%s] is not a recognized validator type and has not been activated."), *DefinitionPair.Key.ToString());
 		}
 	}
 
+
 	ValidationService->InitializeTasks(Tasks);
 }
 
