Use Regex instead of folders and extensions for identifying CL suitability for validators

diff --git a/Engine/Source/Programs/SubmitTool/Private/Logic/Validators/ValidatorBase.cpp b/Engine/Source/Programs/SubmitTool/Private/Logic/Validators/ValidatorBase.cpp
--- a/Engine/Source/Programs/SubmitTool/Private/Logic/Validators/ValidatorBase.cpp
+++ b/Engine/Source/Programs/SubmitTool/Private/Logic/Validators/ValidatorBase.cpp
@@ -2,6 +2,7 @@
 
 #include "ValidatorBase.h"
 #include "HAL/FileManagerGeneric.h"
+#include "Internationalization/Regex.h"
 #include "Misc/Paths.h"
 #include "Misc/Timespan.h"
 #include "AnalyticsEventAttribute.h"
@@ -70,9 +71,17 @@ void FValidatorBase::StartValidation()
 		}
 		else
 		{
-			const FString Extensions = Definition->IncludeFilesWithExtension.IsEmpty() ? TEXT(".*") : *FString::Join(Definition->IncludeFilesWithExtension, TEXT("|"));
-			UE_LOG(LogValidators, Log, TEXT("[%s] No files match the filter %s{%s} %s doesn't need to run"), *ValidatorName, *Definition->IncludeFilesInDirectory, *Extensions, *ValidatorName);
-			UE_LOG(LogValidatorsResult, Log, TEXT("[%s] No files match the filter %s{%s} %s doesn't need to run"), *ValidatorName, *Definition->IncludeFilesInDirectory, *Extensions, *ValidatorName);
+			if(!Definition->AppliesToCLRegex.IsEmpty())
+			{
+				UE_LOG(LogValidators, Log, TEXT("[%s] No files match the regex %s. %s"), *ValidatorName, *FConfiguration::Substitute(Definition->AppliesToCLRegex), *Definition->NotApplicableToCLMessage);
+				UE_LOG(LogValidatorsResult, Log, TEXT("[%s] No files match the regex %s. %s"), *ValidatorName, *FConfiguration::Substitute(Definition->AppliesToCLRegex), *Definition->NotApplicableToCLMessage);
+			}
+			else
+			{
+				const FString Extensions = Definition->IncludeFilesWithExtension.IsEmpty() ? TEXT(".*") : *FString::Join(Definition->IncludeFilesWithExtension, TEXT("|"));
+				UE_LOG(LogValidators, Log, TEXT("[%s] No files match the filter %s{%s} %s doesn't need to run"), *ValidatorName, *Definition->IncludeFilesInDirectory, *Extensions, *ValidatorName);
+				UE_LOG(LogValidatorsResult, Log, TEXT("[%s] No files match the filter %s{%s} %s doesn't need to run"), *ValidatorName, *Definition->IncludeFilesInDirectory, *Extensions, *ValidatorName);
+			}
 		}
 		Skip();
 	}
@@ -352,29 +361,40 @@ bool FValidatorBase::AppliesToFile(const FSourceControlStateRef InFile, bool Inb
 		FString Filename = InFile->GetFilename();
 		FPaths::NormalizeFilename(Filename);
 
-		if(!Definition->IncludeFilesInDirectory.IsEmpty())
+		if (!Definition->AppliesToCLRegex.IsEmpty())
 		{
-			if(!Filename.StartsWith(Definition->IncludeFilesInDirectory, ESearchCase::IgnoreCase))
-			{
-				return false;
-			}
-		}
+			const FString RegexPat = FConfiguration::Substitute(Definition->AppliesToCLRegex);
+			FRegexPattern Pattern = FRegexPattern(RegexPat, ERegexPatternFlags::CaseInsensitive);
+			FRegexMatcher regex = FRegexMatcher(Pattern, Filename);
 
-		if(Definition->IncludeFilesWithExtension.IsEmpty())
-		{
-			bIncluded = true;
+			bIncluded = regex.FindNext();
 		}
-
-		for(int Idx = 0; Idx < Definition->IncludeFilesWithExtension.Num(); Idx++)
+		else
 		{
-			if(Filename.EndsWith(Definition->IncludeFilesWithExtension[Idx], ESearchCase::IgnoreCase))
+			if(!Definition->IncludeFilesInDirectory.IsEmpty())
+			{
+				if(!Filename.StartsWith(Definition->IncludeFilesInDirectory, ESearchCase::IgnoreCase))
+				{
+					return false;
+				}
+			}
+
+			if(Definition->IncludeFilesWithExtension.IsEmpty())
 			{
 				bIncluded = true;
-				break;
+			}
+
+			for(int Idx = 0; Idx < Definition->IncludeFilesWithExtension.Num(); Idx++)
+			{
+				if(Filename.EndsWith(Definition->IncludeFilesWithExtension[Idx], ESearchCase::IgnoreCase))
+				{
+					bIncluded = true;
+					break;
+				}
 			}
 		}
 
-		if(InbAllowIncremental)
+		if(InbAllowIncremental && bIncluded)
 		{
 			FFileManagerGeneric FileManager;
 			FDateTime LastValidation = ServiceProvider.Pin()->GetService<ICacheDataService>()->GetLastValidationDate(ServiceProvider.Pin()->GetService<FChangelistService>()->GetCLID(), GetValidatorNameId(), GetValidationConfigId(), InFile->GetFilename());
diff --git a/Engine/Source/Programs/SubmitTool/Private/Logic/Validators/ValidatorDefinition.h b/Engine/Source/Programs/SubmitTool/Private/Logic/Validators/ValidatorDefinition.h
--- a/Engine/Source/Programs/SubmitTool/Private/Logic/Validators/ValidatorDefinition.h
+++ b/Engine/Source/Programs/SubmitTool/Private/Logic/Validators/ValidatorDefinition.h
@@ -61,12 +61,27 @@ struct FValidatorDefinition
 	FString CustomName;
 
 	///
+	/// Regex to test paths to check if this validator should be applied
+	/// Replaces IncludeFilesWithExtension & IncludeFilesInDirectory
+	/// 
+	UPROPERTY()
+	FString AppliesToCLRegex;
+
+	///
+	/// Optional message to print when the validator doesn't apply to CL to give more context to AppliesToCLRegex
+	/// 
+	UPROPERTY()
+	FString NotApplicableToCLMessage;
+
+	///
+	/// DEPRECATED: USE AppliesToCLRegex
 	/// Files with any of these extensions will be included
 	/// 
 	UPROPERTY()
 	TArray<FString> IncludeFilesWithExtension;
 
 	///
+	/// DEPRECATED: USE AppliesToCLRegex
 	///	Only run this validator for files under this directory
 	/// 
 	UPROPERTY()
