fixing logic issue in json validator

diff --git a/Engine/Source/Programs/SubmitTool/Private/Logic/Validators/JsonValidator.cpp b/Engine/Source/Programs/SubmitTool/Private/Logic/Validators/JsonValidator.cpp
--- a/Engine/Source/Programs/SubmitTool/Private/Logic/Validators/JsonValidator.cpp
+++ b/Engine/Source/Programs/SubmitTool/Private/Logic/Validators/JsonValidator.cpp
@@ -45,7 +45,7 @@ bool FJsonValidator::Validate(const FString& InCLDescription, const TArray<FSour
 
 				TSharedRef<TJsonReader<>> JsonReader = TJsonReaderFactory<>::Create(JSON);
 				TSharedPtr<FJsonObject> JsonObject;
-				if (!FJsonSerializer::Deserialize(JsonReader, JsonObject) || JsonObject.IsValid())
+				if (!FJsonSerializer::Deserialize(JsonReader, JsonObject) || !JsonObject.IsValid())
 				{
 					const FString& Error = JsonReader->GetErrorMessage();
 					LogFailure(FString::Printf(TEXT("[%s] %s is an invalid JSON file: %s"), *GetValidatorName(), *Filename, *Error));
