Fixing crashes related to p4 operations not being cancelled on submit tool exit
Fixing stream root not being substituted in incompatible files
Increased perforce max concurrency to startup faster

diff --git a/Engine/Source/Programs/SubmitTool/Private/Logic/ChangelistService.cpp b/Engine/Source/Programs/SubmitTool/Private/Logic/ChangelistService.cpp
--- a/Engine/Source/Programs/SubmitTool/Private/Logic/ChangelistService.cpp
+++ b/Engine/Source/Programs/SubmitTool/Private/Logic/ChangelistService.cpp
@@ -114,10 +114,12 @@ void FChangelistService::FindInitialChangelistsAsync()
 
 				// Fire & Forget updating all CLs
 				TSharedRef<FUpdatePendingChangelistsStatus> UpdatePendingChangelistsOperation = ISourceControlOperation::Create<FUpdatePendingChangelistsStatus>();
+				ActiveP4Operations.Add(UpdatePendingChangelistsOperation);
 				UpdatePendingChangelistsOperation->SetUpdateFilesStates(true);
 				UpdatePendingChangelistsOperation->SetUpdateShelvedFilesStates(true);
 				UpdatePendingChangelistsOperation->SetUpdateAllChangelists(true);
-				CachedSCCProvider->Execute(UpdatePendingChangelistsOperation, EConcurrency::Asynchronous);
+				CachedSCCProvider->Execute(UpdatePendingChangelistsOperation, EConcurrency::Asynchronous, FSourceControlOperationComplete::CreateLambda(
+				[this, UpdatePendingChangelistsOperation](const FSourceControlOperationRef& Operation, ECommandResult::Type Result) { ActiveP4Operations.Remove(UpdatePendingChangelistsOperation); }));
 			}
 			else
 			{
@@ -484,7 +486,7 @@ void FChangelistService::CancelP4Operations(FName OperationName)
 {
 	for(const TSharedRef<ISourceControlOperation>& Operation : ActiveP4Operations)
 	{
-		if(OperationName == FName() || Operation->GetName().IsEqual(OperationName, ENameCase::IgnoreCase))
+		if(OperationName.IsNone() || Operation->GetName().IsEqual(OperationName, ENameCase::IgnoreCase))
 		{
 			if(CachedSCCProvider->CanCancelOperation(Operation))
 			{
diff --git a/Engine/Source/Programs/SubmitTool/Private/Logic/Services/SourceControl/SubmitToolPerforce.cpp b/Engine/Source/Programs/SubmitTool/Private/Logic/Services/SourceControl/SubmitToolPerforce.cpp
--- a/Engine/Source/Programs/SubmitTool/Private/Logic/Services/SourceControl/SubmitToolPerforce.cpp
+++ b/Engine/Source/Programs/SubmitTool/Private/Logic/Services/SourceControl/SubmitToolPerforce.cpp
@@ -18,7 +18,7 @@ constexpr size_t MaxRecentUsers = 12;
 constexpr size_t RecentUsersDatVersion = 1;
 constexpr size_t MaxRecentGroups = 12;
 constexpr size_t RecentGroupsDatVersion = 1;
-constexpr size_t MaxConnections = 1;
+constexpr size_t MaxConnections = 7;
 constexpr size_t MaxConnectionsAttempts = 10;
 
 
@@ -37,6 +37,7 @@ FP4Connection::~FP4Connection()
 
 FConnectionPool::~FConnectionPool()
 {
+	RequestCancel();
 	for(TUniquePtr<FClientApiWrapper>& ConWrapper : P4Connections)
 	{
 		Error P4Error;
@@ -65,6 +66,7 @@ TUniquePtr<FP4Connection> FConnectionPool::GetAvailableConnection()
 
 	if(P4Connections.Num() < MaxConnections)
 	{
+		UE_LOG(LogSubmitToolP4Debug, Log, TEXT("Creating new p4 connection %d/%d."), P4Connections.Num()+1, MaxConnections);
 		size_t Idx = CreateConnection();
 		if(Idx != -1)
 		{
@@ -100,7 +102,7 @@ size_t FConnectionPool::CreateConnection()
 	}
 
 	FScopeLock Lock(&Mutex);
-	return P4Connections.Add(MakeUnique<FClientApiWrapper>(MoveTemp(P4Client)));
+	return P4Connections.Add(MakeUnique<FClientApiWrapper>(MoveTemp(P4Client), FOnIsCancelled::CreateLambda([this] { return bWantsCancel; })));
 }
 
 FSubmitToolPerforce::FSubmitToolPerforce(const FSubmitToolParameters& InParameters)
diff --git a/Engine/Source/Programs/SubmitTool/Private/Logic/Services/SourceControl/SubmitToolPerforce.h b/Engine/Source/Programs/SubmitTool/Private/Logic/Services/SourceControl/SubmitToolPerforce.h
--- a/Engine/Source/Programs/SubmitTool/Private/Logic/Services/SourceControl/SubmitToolPerforce.h
+++ b/Engine/Source/Programs/SubmitTool/Private/Logic/Services/SourceControl/SubmitToolPerforce.h
@@ -11,13 +11,34 @@
 class FProcessWrapper;
 class ISourceControlProvider;
 
-struct FClientApiWrapper
+DECLARE_DELEGATE_RetVal(bool, FOnIsCancelled);
+
+class FP4KeepAlive : public KeepAlive
 {
-	FClientApiWrapper(TUniquePtr<ClientApi>&& InConnection)
-	: bIsReady(true), Connection(Forward<TUniquePtr<ClientApi>>(InConnection))
+public:
+	FP4KeepAlive(FOnIsCancelled InIsCancelled)
+		: IsCancelled(InIsCancelled)
 	{}
 
+	/** Called when the perforce connection wants to know if it should stay connected */
+	virtual int IsAlive() override
+	{
+		return 0;
+	}
+
+	FOnIsCancelled IsCancelled;
+};
+
+struct FClientApiWrapper
+{
+	FClientApiWrapper(TUniquePtr<ClientApi>&& InConnection, FOnIsCancelled InKeepAlive)
+	: bIsReady(true), KeepAlive(InKeepAlive), Connection(Forward<TUniquePtr<ClientApi>>(InConnection))
+	{
+		Connection->SetBreak(&KeepAlive);
+	}
+
 	bool bIsReady = true;
+	FP4KeepAlive KeepAlive;
 	TUniquePtr<ClientApi> Connection;
 };
 
@@ -41,12 +62,13 @@ public:
 	~FConnectionPool();
 	TUniquePtr<FP4Connection> GetAvailableConnection();
 	
+	void RequestCancel(){ bWantsCancel = true; }
 private:
 	size_t CreateConnection();
+	bool bWantsCancel = false;
 	bool bConnectionFailed = false;
 	TArray<TUniquePtr<FClientApiWrapper>> P4Connections;
 	FCriticalSection Mutex;
-	const int32 MaxConnections = 5;
 };
 
 class FSubmitToolPerforce final : public ISTSourceControlService
diff --git a/Engine/Source/Programs/SubmitTool/Private/Models/ModelInterface.cpp b/Engine/Source/Programs/SubmitTool/Private/Models/ModelInterface.cpp
--- a/Engine/Source/Programs/SubmitTool/Private/Models/ModelInterface.cpp
+++ b/Engine/Source/Programs/SubmitTool/Private/Models/ModelInterface.cpp
@@ -479,6 +479,8 @@ void FModelInterface::StartSubmitProcess()
 		{
 			const TArray<FString>& FilesInCL = ChangelistService->GetFilesDepotPaths();
 
+			FConfiguration::AddOrUpdateEntry(TEXT("$(StreamRoot)"), *GetRootStreamName());
+
 			for(const FIncompatibleFilesGroup& FileGroup : Parameters.IncompatibleFilesParams.IncompatibleFileGroups)
 			{
 				TArray<size_t, TInlineAllocator<8>> Indexes;
