ST no longer require that preflight validator is finished to request integrations

diff --git a/Engine/Source/Programs/SubmitTool/Private/Logic/TasksService.cpp b/Engine/Source/Programs/SubmitTool/Private/Logic/TasksService.cpp
--- a/Engine/Source/Programs/SubmitTool/Private/Logic/TasksService.cpp
+++ b/Engine/Source/Programs/SubmitTool/Private/Logic/TasksService.cpp
@@ -260,7 +260,7 @@ bool FTasksService::AreTasksPendingQueue() const
 }
 
 
-bool FTasksService::GetIsRunSuccessful() const
+bool FTasksService::GetIsRunSuccessful(bool bWaitForOptionalCompletes) const
 {
 	for(const TPair<FName, TSharedPtr<FValidatorBase>>& TaskPair : Tasks)
 	{
@@ -273,7 +273,7 @@ bool FTasksService::GetIsRunSuccessful() const
 		}
 		else
 		{
-			if (TaskPair.Value->Definition->bRequireCompleteWhenOptional && TaskPair.Value->GetIsRunningOrQueued())
+			if (bWaitForOptionalCompletes && TaskPair.Value->Definition->bRequireCompleteWhenOptional && TaskPair.Value->GetIsRunningOrQueued())
 			{
 				return false;
 			}
diff --git a/Engine/Source/Programs/SubmitTool/Private/Logic/TasksService.h b/Engine/Source/Programs/SubmitTool/Private/Logic/TasksService.h
--- a/Engine/Source/Programs/SubmitTool/Private/Logic/TasksService.h
+++ b/Engine/Source/Programs/SubmitTool/Private/Logic/TasksService.h
@@ -32,7 +32,7 @@ public:
 	virtual void CheckForLocalFileEdit();
 	virtual void CheckForTagSkips();
 
-	virtual bool GetIsRunSuccessful() const;
+	virtual bool GetIsRunSuccessful(bool bWaitForOptionalCompletes = true) const;
 	virtual bool GetIsAnyTaskRunning() const;
 	virtual bool AreTasksPendingQueue() const;
 
diff --git a/Engine/Source/Programs/SubmitTool/Private/Models/ModelInterface.h b/Engine/Source/Programs/SubmitTool/Private/Models/ModelInterface.h
--- a/Engine/Source/Programs/SubmitTool/Private/Models/ModelInterface.h
+++ b/Engine/Source/Programs/SubmitTool/Private/Models/ModelInterface.h
@@ -129,7 +129,7 @@ public:
 	void ValidateChangelist() const													{ ValidationService->QueueAll(); }
 	void ValidateSingle(const FName& ValidatorId, bool bForce = true)				{ ValidationService->QueueSingle(ValidatorId, bForce); }
 	void ValidateCLDescription() const												{ ValidationService->StopTasksByArea(ETaskArea::Changelist); ValidationService->QueueByArea(ETaskArea::Changelist); }
-	bool IsCLValid() const															{ return ValidationService->GetIsRunSuccessful(); }
+	bool IsCLValid() const															{ return ValidationService->GetIsRunSuccessful(!IsIntegrationRequired()); }
 	bool CanLaunchPreflight() const;
 	void ReevaluateSubmitToolTag();
 	void UpdateSubmitToolTag(bool InbAdd);
