Collapsed pre-submit operations that are inactive

diff --git a/Engine/Source/Programs/SubmitTool/Private/SubmitToolApp.cpp b/Engine/Source/Programs/SubmitTool/Private/SubmitToolApp.cpp
--- a/Engine/Source/Programs/SubmitTool/Private/SubmitToolApp.cpp
+++ b/Engine/Source/Programs/SubmitTool/Private/SubmitToolApp.cpp
@@ -69,7 +69,6 @@ int RunSubmitTool(const TCHAR* CommandLine, const FGuid& InSessionID)
 	FAppStyle::SetAppStyleSet(FSubmitToolStyle::Get());
 	FSlateApplication::InitHighDPI(true);
 
-
 	// App Scope
 	{
 		UE_LOG(LogSubmitTool, Log, TEXT("Version %s"), *FAppVersion::GetVersion());
@@ -157,26 +156,37 @@ int RunSubmitTool(const TCHAR* CommandLine, const FGuid& InSessionID)
 				Tab->GetParentWindow()->Maximize();
 			}
 
+			double DeltaTime = 0.0;
+			double LastTime = FPlatformTime::Seconds();
+			const float IdealFrameTime = 1.0f / 60;
+			const float BackgroundFrameTime = 1.0f / 4;
+
 			// Loop the engine
-			while(!IsEngineExitRequested())
+			while (!IsEngineExitRequested())
 			{
 				BeginExitIfRequested();
 
 				FTaskGraphInterface::Get().ProcessThreadUntilIdle(ENamedThreads::GameThread);
-				UE::Stats::FStats::AdvanceFrame(false);
-				FTSTicker::GetCoreTicker().Tick(FSlateApplication::Get().GetDeltaTime());
+				FTSTicker::GetCoreTicker().Tick(DeltaTime);
 				FSlateApplication::Get().PumpMessages();
 				FSlateApplication::Get().Tick();
 
-				if(IsEngineExitRequested())
+				if (IsEngineExitRequested())
 				{
 					// Dispose here so SCCProvider has time to clean up
 					ModelInterface->Dispose();
 				}
 
-				FPlatformProcess::Sleep(0.01f);
-				GFrameCounter++;
+				const float FrameTime = !FPlatformApplicationMisc::IsThisApplicationForeground() && (FPlatformTime::Seconds() - FSlateApplication::Get().GetLastUserInteractionTime()) > 5 ? BackgroundFrameTime : IdealFrameTime;
+				FPlatformProcess::Sleep(FMath::Max<float>(0.f, FrameTime - static_cast<float>(FPlatformTime::Seconds() - LastTime)));
 
+				DeltaTime = FPlatformTime::Seconds() - LastTime;
+				LastTime = FPlatformTime::Seconds();
+
+				UE::Stats::FStats::AdvanceFrame(false);
+				FCoreDelegates::OnEndFrame.Broadcast();
+
+				GFrameCounter++;
 			}
 		}
 		// ensure all the telemetry events are flushed before unloading modules
diff --git a/Engine/Source/Programs/SubmitTool/Private/View/Widgets/ValidatorsWidget.cpp b/Engine/Source/Programs/SubmitTool/Private/View/Widgets/ValidatorsWidget.cpp
--- a/Engine/Source/Programs/SubmitTool/Private/View/Widgets/ValidatorsWidget.cpp
+++ b/Engine/Source/Programs/SubmitTool/Private/View/Widgets/ValidatorsWidget.cpp
@@ -87,13 +87,19 @@ TSharedRef<SWidget> SValidatorsWidget::BuildValidatorsView(bool bListPreSubmitOp
 
 	for(const TPair<FName, TArray<TWeakPtr<const FValidatorBase>>>& Pair : Validators)
 	{
+		bool bInitiallyCollapsed = false;
+		if (!(bListPreSubmitOperations && Pair.Key.IsEqual(ActiveSection)))
+		{
+			bInitiallyCollapsed = !FSubmitToolUserPrefs::Get()->UISectionExpandState.FindOrAdd(Pair.Key, false);
+		}
+
 		VBox->AddSlot()
 			.Padding(FMargin(0.f, 2.f))
 			.AutoHeight()
 			.AttachWidget
 			(
 				SNew(SExpandableArea)
-				.InitiallyCollapsed(!bListPreSubmitOperations && !FSubmitToolUserPrefs::Get()->UISectionExpandState.FindOrAdd(Pair.Key, false))
+				.InitiallyCollapsed(bInitiallyCollapsed)
 				.BorderImage(FAppStyle::Get().GetBrush("Brushes.Header"))
 				.BodyBorderImage(FAppStyle::Get().GetBrush("Brushes.Recessed"))
 				.OnAreaExpansionChanged_Lambda([DisplayName = Pair.Key](bool bExpanded) { FSubmitToolUserPrefs::Get()->UISectionExpandState.FindOrAdd(DisplayName, false) = bExpanded; })
