Add support for disabling validators and properties to make validation dependent on a file in the folder tree
Removing space in Integration window
Integration doesn't require to wait for optional validators finishing anymore
Adding Platforms to NonUnity dropdown

diff --git a/Engine/Source/Programs/SubmitTool/Private/Logic/Validators/ValidatorBase.cpp b/Engine/Source/Programs/SubmitTool/Private/Logic/Validators/ValidatorBase.cpp
--- a/Engine/Source/Programs/SubmitTool/Private/Logic/Validators/ValidatorBase.cpp
+++ b/Engine/Source/Programs/SubmitTool/Private/Logic/Validators/ValidatorBase.cpp
@@ -394,6 +394,30 @@ bool FValidatorBase::AppliesToFile(const FSourceControlStateRef InFile, bool Inb
 			}
 		}
 
+		if(bIncluded && (!Definition->RequireFileInHierarchy.IsEmpty() || !Definition->ExcludeWhenFileInHierarchy.IsEmpty()))
+		{
+			bIncluded = false;
+			FString CurrentDir = FPaths::GetPath(Filename);
+
+			while (!CurrentDir.IsEmpty() && IFileManager::Get().DirectoryExists(*CurrentDir))
+			{
+				TArray<FString> IncludeFilesFound;
+				IFileManager::Get().FindFiles(IncludeFilesFound, *(CurrentDir / Definition->RequireFileInHierarchy), true, false);
+
+				TArray<FString> ExcludeFilesFound;
+				IFileManager::Get().FindFiles(ExcludeFilesFound, *(CurrentDir / Definition->ExcludeWhenFileInHierarchy), true, false);
+
+
+				if (IncludeFilesFound.Num() > 0 || ExcludeFilesFound.Num() > 0)
+				{		
+					bIncluded = ExcludeFilesFound.Num() == 0;
+					break;
+				}
+				
+				CurrentDir = FPaths::GetPath(CurrentDir);
+			}
+		}
+
 		if(InbAllowIncremental && bIncluded)
 		{
 			FFileManagerGeneric FileManager;
diff --git a/Engine/Source/Programs/SubmitTool/Private/Logic/Validators/ValidatorDefinition.h b/Engine/Source/Programs/SubmitTool/Private/Logic/Validators/ValidatorDefinition.h
--- a/Engine/Source/Programs/SubmitTool/Private/Logic/Validators/ValidatorDefinition.h
+++ b/Engine/Source/Programs/SubmitTool/Private/Logic/Validators/ValidatorDefinition.h
@@ -35,7 +35,13 @@ struct FValidatorDefinition
 	///
 	UPROPERTY()
 	FString Type;
-	
+
+	///
+	///	If true, this validator will be disabled and not loaded, useful for config overriding a parent validator
+	/// 
+	UPROPERTY()
+	bool bIsDisabled = false;
+
 	/// 
 	/// Whether the validator is required to allow submission or not, a failing required validation will always block submission
 	/// 
@@ -86,7 +92,21 @@ struct FValidatorDefinition
 	/// 
 	UPROPERTY()
 	FString IncludeFilesInDirectory;
+	
+	///
+	/// Wildcard pattern: only apply this validator if a file matching this pattern is somewhere in the path parent structure
+	/// i.e: *.uproject -> validator will only run if there is a uproject in their parent folder structure
+	/// 
+	UPROPERTY()
+	FString RequireFileInHierarchy;
 
+	///
+	/// Wildcard pattern: Opposite of RequireFileInHierarchy this validator will never run if a file matching this pattern is somewhere in the path parent structure
+	/// i.e: *.uproject -> validator will never run if there is a uproject in their parent folder structure
+	/// 
+	UPROPERTY()
+	FString ExcludeWhenFileInHierarchy;
+	
 	///
 	///	This text will be added to the description if this validation passes
 	/// 
diff --git a/Engine/Source/Programs/SubmitTool/Private/Models/ModelInterface.cpp b/Engine/Source/Programs/SubmitTool/Private/Models/ModelInterface.cpp
--- a/Engine/Source/Programs/SubmitTool/Private/Models/ModelInterface.cpp
+++ b/Engine/Source/Programs/SubmitTool/Private/Models/ModelInterface.cpp
@@ -205,6 +205,12 @@ void FModelInterface::ParseValidators() const
 			continue;
 		}
 
+		if (TaskDefinition.bIsDisabled)
+		{
+			UE_LOG(LogSubmitToolDebug, Log, TEXT("Task %s was disabled by configuration"), *DefinitionPair.Key.ToString());
+			continue;
+		}
+
 		if(TaskDefinition.Type.Equals(SubmitToolParseConstants::TagValidator, ESearchCase::IgnoreCase))
 		{
 			Tasks.Add(MakeShared<FTagValidator>(DefinitionPair.Key, Parameters, ServiceProvider.ToSharedRef(), DefinitionPair.Value));
@@ -340,7 +346,7 @@ bool FModelInterface::CanLaunchPreflight() const
 
 void FModelInterface::ReevaluateSubmitToolTag()
 {
-	UpdateSubmitToolTag(ValidationService->GetIsRunSuccessful());	
+	UpdateSubmitToolTag(ValidationService->GetIsRunSuccessful(!IsIntegrationRequired()));	
 }
 
 void FModelInterface::UpdateSubmitToolTag(bool InbAdd)
