Avoid regex testing every path per frame

diff --git a/Engine/Source/Programs/SubmitTool/Private/View/Widgets/ValidatorsWidget.cpp b/Engine/Source/Programs/SubmitTool/Private/View/Widgets/ValidatorsWidget.cpp
--- a/Engine/Source/Programs/SubmitTool/Private/View/Widgets/ValidatorsWidget.cpp
+++ b/Engine/Source/Programs/SubmitTool/Private/View/Widgets/ValidatorsWidget.cpp
@@ -206,12 +206,14 @@ class SValidatorNode : public SMultiColumnTableRow<TWeakPtr<const FValidatorBase
 public:
 	SLATE_BEGIN_ARGS(SValidatorNode) {}
 		SLATE_ARGUMENT(FModelInterface*, ModelInterface)
+		SLATE_ARGUMENT(bool, IsRelevantToCL)
 		SLATE_EVENT(FOnViewValidatorLog, ViewLogCallback)
 	SLATE_END_ARGS()
 
 	void Construct(const FArguments& InArgs, const TSharedRef<STableViewBase>& InOwnerTable, TWeakPtr<const FValidatorBase> InNode)
 	{
 		ModelInterface = InArgs._ModelInterface;
+		bIsRelevantToCL = InArgs._IsRelevantToCL;
 		Node = InNode;
 		ViewLogCallback = InArgs._ViewLogCallback;
 		Options = TMap<TSharedPtr<FString>, TArray<TSharedPtr<FString>>>();
@@ -323,8 +325,6 @@ public:
 				Required
 			);
 		}
-
-
 		else if(InColumnName == ValidatorsColumns::Status)
 		{
 			box->SetContent(
@@ -358,7 +358,7 @@ public:
 				.Padding(3,0,0,0)
 				[
 					SNew(SButton)
-					.IsEnabled_Lambda([this]() {return ModelInterface->GetInputEnabled() && Node.Pin()->IsRelevantToCL(); })
+					.IsEnabled_Lambda([this]() {return ModelInterface->GetInputEnabled() && bIsRelevantToCL; })
 					.HAlign(EHorizontalAlignment::HAlign_Center)
 					.VAlign(EVerticalAlignment::VAlign_Center)
 					.OnClicked_Lambda([this]()
@@ -378,7 +378,7 @@ public:
 				.Padding(3,0,0,0)
 				[
 					SNew(SButton)
-					.IsEnabled_Lambda([this]() {return ModelInterface->GetInputEnabled() && Node.Pin()->IsRelevantToCL(); })
+					.IsEnabled_Lambda([this]() {return ModelInterface->GetInputEnabled() && bIsRelevantToCL; })
 					.HAlign(EHorizontalAlignment::HAlign_Center)
 					.VAlign(EVerticalAlignment::VAlign_Center)
 					.OnClicked_Lambda([this]()
@@ -423,6 +423,7 @@ public:
 
 	TWeakPtr<const FValidatorBase> Node;
 	FModelInterface* ModelInterface;
+	bool bIsRelevantToCL = false;
 	FOnViewValidatorLog ViewLogCallback;
 };
 
@@ -430,6 +431,7 @@ TSharedRef<ITableRow> SValidatorsWidget::GenerateRow(TWeakPtr<const FValidatorBa
 {
 	return SNew(SValidatorNode, OwnerTable, InItem)
 		.ModelInterface(ModelInterface)
+		.IsRelevantToCL(InItem.Pin()->IsRelevantToCL())
 		.ViewLogCallback(OnViewLog);
 }
 
