Adding regex filtering to skip lines from json parsing based on config

diff --git a/Engine/Source/Programs/SubmitTool/Private/Logic/Validators/JsonValidator.cpp b/Engine/Source/Programs/SubmitTool/Private/Logic/Validators/JsonValidator.cpp
--- a/Engine/Source/Programs/SubmitTool/Private/Logic/Validators/JsonValidator.cpp
+++ b/Engine/Source/Programs/SubmitTool/Private/Logic/Validators/JsonValidator.cpp
@@ -2,6 +2,7 @@
 
 #include "JsonValidator.h"
 
+#include "Internationalization/Regex.h"
 #include "Serialization/JsonReader.h"
 #include "Serialization/JsonSerializer.h"
 #include "Misc/FileHelper.h"
@@ -12,6 +13,7 @@
 FJsonValidator::FJsonValidator(const FName& InNameId, const FSubmitToolParameters& InParameters, TSharedRef<FSubmitToolServiceProvider> InServiceProvider, const FString& InDefinition) :
 	FValidatorBase(InNameId, InParameters, InServiceProvider, InDefinition)
 {
+	ParseDefinition(InDefinition);
 }
 
 bool FJsonValidator::Validate(const FString& InCLDescription, const TArray<FSourceControlStateRef>& InFilteredFilesInCL, const TArray<const FTag*>& InTags)
@@ -36,14 +38,32 @@ bool FJsonValidator::Validate(const FString& InCLDescription, const TArray<FSour
 				}
 
 				FString JSON;
-				if (!FFileHelper::LoadFileToString(JSON, *Filename))
+				const FString& RegexFiltering = GetTypedDefinition<FJsonValidatorDefinition>()->RegexLineExclusion;
+				if (!RegexFiltering.IsEmpty())
 				{
-					LogFailure(FString::Printf(TEXT("[%s] %s could not be loaded"), *GetValidatorName(), *Filename));
-					bValid = false;
-					continue;
+					FRegexPattern Pattern = FRegexPattern(RegexFiltering, ERegexPatternFlags::CaseInsensitive);
+
+					TArray<FString> Lines;
+					if (!FFileHelper::LoadFileToStringArrayWithPredicate(Lines, *Filename, [&Pattern](const FString& Line) { return !FRegexMatcher(Pattern, Line).FindNext();}))
+					{
+						LogFailure(FString::Printf(TEXT("[%s] %s could not be loaded"), *GetValidatorName(), *Filename));
+						bValid = false;
+						continue;
+					}
+
+					JSON = FString::Join(Lines, TEXT("\n"));
+				}
+				else
+				{
+					if (!FFileHelper::LoadFileToString(JSON, *Filename))
+					{
+						LogFailure(FString::Printf(TEXT("[%s] %s could not be loaded"), *GetValidatorName(), *Filename));
+						bValid = false;
+						continue;
+					}
 				}
 
-				TSharedRef<TJsonReader<>> JsonReader = TJsonReaderFactory<>::Create(JSON);
+				TSharedRef<TJsonReader<>> JsonReader = TJsonReaderFactory<>::Create(MoveTemp(JSON));
 				TSharedPtr<FJsonValue> JsonValue;
 				if (!FJsonSerializer::Deserialize(JsonReader, JsonValue) || !JsonValue.IsValid())
 				{
@@ -59,6 +79,22 @@ bool FJsonValidator::Validate(const FString& InCLDescription, const TArray<FSour
 	return true;
 }
 
+void FJsonValidator::ParseDefinition(const FString& InDefinition)
+{
+	FStringOutputDevice Errors;
+
+	FJsonValidatorDefinition* DefinitionToLoad = new FJsonValidatorDefinition;
+	FJsonValidatorDefinition::StaticStruct()->ImportText(*InDefinition, DefinitionToLoad, nullptr, 0, &Errors, FJsonValidatorDefinition::StaticStruct()->GetName());
+
+	Definition.Reset(DefinitionToLoad);
+
+	if (!Errors.IsEmpty())
+	{
+		UE_LOG(LogSubmitTool, Error, TEXT("[%s] Error loading parameter file %s"), *GetValidatorName(), *Errors);
+		FModelInterface::SetErrorState();
+	}
+}
+
 void FJsonValidator::Tick(float InDeltaTime)
 {
 	FValidatorBase::Tick(InDeltaTime);
diff --git a/Engine/Source/Programs/SubmitTool/Private/Logic/Validators/JsonValidator.h b/Engine/Source/Programs/SubmitTool/Private/Logic/Validators/JsonValidator.h
--- a/Engine/Source/Programs/SubmitTool/Private/Logic/Validators/JsonValidator.h
+++ b/Engine/Source/Programs/SubmitTool/Private/Logic/Validators/JsonValidator.h
@@ -15,6 +15,8 @@ public:
 	{
 		return SubmitToolParseConstants::JsonValidator;
 	}
+protected:
+	virtual void ParseDefinition(const FString& InDefinition) override;
 
 private:
 	virtual void Tick(float InDeltaTime) override;
diff --git a/Engine/Source/Programs/SubmitTool/Private/Logic/Validators/ValidatorDefinition.h b/Engine/Source/Programs/SubmitTool/Private/Logic/Validators/ValidatorDefinition.h
--- a/Engine/Source/Programs/SubmitTool/Private/Logic/Validators/ValidatorDefinition.h
+++ b/Engine/Source/Programs/SubmitTool/Private/Logic/Validators/ValidatorDefinition.h
@@ -329,6 +329,18 @@ struct FPackageDataValidatorDefinition : public FValidatorDefinition
 	TArray<FString> ExcludedExtensions;
 };
 
+USTRUCT()
+struct FJsonValidatorDefinition : public FValidatorDefinition
+{
+	GENERATED_BODY()
+
+	///
+	/// Parse the lines of the json and do not include the ones that match this regex (for custom parsing)
+	/// 
+	UPROPERTY()
+	FString RegexLineExclusion;
+};
+
 USTRUCT()
 struct FVirtualizationToolDefinition : public FValidatorRunExecutableDefinition
 {
