Fix trying to read prelfights on horde timeout

diff --git a/Engine/Source/Programs/SubmitTool/Private/Logic/PreflightService.cpp b/Engine/Source/Programs/SubmitTool/Private/Logic/PreflightService.cpp
--- a/Engine/Source/Programs/SubmitTool/Private/Logic/PreflightService.cpp
+++ b/Engine/Source/Programs/SubmitTool/Private/Logic/PreflightService.cpp
@@ -363,7 +363,7 @@ void FPreflightService::FetchPreflightInfo(bool bRequeue, const FString& InOAuth
 
 				if(ActiveUnlinkedRequests == 0)
 				{
-					if(OnPreflightDataUpdated.IsBound())
+					if(OnPreflightDataUpdated.IsBound() && HordePreflights.IsValid())
 					{
 						OnPreflightDataUpdated.Broadcast(HordePreflights, UnlinkedHordePreflights);
 					}
@@ -491,7 +491,7 @@ void FPreflightService::FetchUnlinkedPreflight(const FString& InPreflightId, boo
 
 			if(ActiveUnlinkedRequests == 0)
 			{
-				if(OnPreflightDataUpdated.IsBound())
+				if(OnPreflightDataUpdated.IsBound() && HordePreflights.IsValid())
 				{
 					OnPreflightDataUpdated.Broadcast(HordePreflights, UnlinkedHordePreflights);
 				}
