Open Jira tickets on submission

diff --git a/Engine/Source/Programs/SubmitTool/Private/Logic/TagService.h b/Engine/Source/Programs/SubmitTool/Private/Logic/TagService.h
--- a/Engine/Source/Programs/SubmitTool/Private/Logic/TagService.h
+++ b/Engine/Source/Programs/SubmitTool/Private/Logic/TagService.h
@@ -34,6 +34,20 @@ public:
 		return nullptr;
 	};
 
+	FTag* GetTagOfType(const FString& InType)
+	{
+		for (const TPair<FString, FTag>& Tag : RegisteredTags)
+		{
+			if (Tag.Value.Definition.InputType.Equals(InType, ESearchCase::IgnoreCase))
+			{
+				return &RegisteredTags[Tag.Key];
+			}
+		}
+
+		return nullptr;
+	}
+
+
 	FTag* GetTagOfSubtype(const FString& InTagSubtype)
 	{
 		for (const TPair<FString, FTag>& Tag : RegisteredTags)
diff --git a/Engine/Source/Programs/SubmitTool/Private/Models/ModelInterface.cpp b/Engine/Source/Programs/SubmitTool/Private/Models/ModelInterface.cpp
--- a/Engine/Source/Programs/SubmitTool/Private/Models/ModelInterface.cpp
+++ b/Engine/Source/Programs/SubmitTool/Private/Models/ModelInterface.cpp
@@ -799,6 +799,18 @@ void FModelInterface::OnSubmitOperationComplete(const FSourceControlOperationRef
 	{
 		// We've submitted, or tried to submit and failed so we only let the user close the app
 		ChangeState(ESubmitToolAppState::Finished);
+		FTag* JiraTag = TagService->GetTagOfType(TEXT("JiraIssue"));
+		if (JiraTag != nullptr && JiraTag->GetValues().Num() != 0)
+		{
+			for (const FString& JiraValue : JiraTag->GetValues())
+			{
+				if (!JiraValue.Equals(TEXT("none"), ESearchCase::IgnoreCase) && !JiraValue.Equals(TEXT("nojira"), ESearchCase::IgnoreCase))
+				{
+					FString Url = FString::Printf(TEXT("https://%s/browse/%s}"), *Parameters.JiraParameters.ServerAddress, *JiraValue);
+					FPlatformProcess::LaunchURL(*Url, nullptr, nullptr);
+				}
+			}
+		}
 
 		if(FSubmitToolUserPrefs::Get()->bCloseOnSubmit)
 		{
