Activation errors logging delayed until tried to run
Reevaluating validator activation on run to be able to recover

diff --git a/Engine/Source/Programs/SubmitTool/Private/Logic/Validators/ValidatorBase.cpp b/Engine/Source/Programs/SubmitTool/Private/Logic/Validators/ValidatorBase.cpp
--- a/Engine/Source/Programs/SubmitTool/Private/Logic/Validators/ValidatorBase.cpp
+++ b/Engine/Source/Programs/SubmitTool/Private/Logic/Validators/ValidatorBase.cpp
@@ -57,10 +57,37 @@ void FValidatorBase::StartValidation()
 	TArray<FSourceControlStateRef> IncrementallySkippedFiles;
 	bool bAppliesToCL = AppliesToCL(ChangelistService->GetCLDescription(), ChangelistService->GetFilesInCL(), TagService->GetTagsArray(), FilteredFiles, IncrementallySkippedFiles, bIncrementalValidation);
 
-	if(bAppliesToCL && !bIsValidSetup)
+	if(bAppliesToCL)
 	{
-		LogFailure(FString::Printf(TEXT("[%s] Task is not correctly setup and should run in this CL"), *ValidatorName));
-		ValidationFinished(false);
+		if (!bIsValidSetup)
+		{
+			// Try to recover if things have changed since startup
+			ActivationErrors.Empty();
+			if (!Activate())
+			{
+				LogFailure(FString::Printf(TEXT("[%s] Task is not correctly setup and should run in this CL"), *ValidatorName));
+				for (const FString& ActivationError : ActivationErrors)
+				{
+					LogFailure(ActivationError);
+				}
+
+				ValidationFinished(false);
+			}
+		}
+
+		if (bIsValidSetup)
+		{
+			if (IncrementallySkippedFiles.Num() != 0)
+			{
+				const FString FileList = FString::JoinBy(IncrementallySkippedFiles, TEXT("\n"), [](const FSourceControlStateRef& InFile) { return InFile->GetFilename(); });
+				UE_LOG(LogValidators, Log, TEXT("[%s] Skipping Files because they were already validated in a previous execution:\n%s"), *GetValidatorName(), *FileList);
+			}
+
+			if (!Validate(ChangelistService->GetCLDescription(), FilteredFiles, TagService->GetTagsArray()))
+			{
+				ValidationFinished(false);
+			}
+		}
 	}
 	else if(!bAppliesToCL)
 	{
@@ -85,19 +112,6 @@ void FValidatorBase::StartValidation()
 		}
 		Skip();
 	}
-	else
-	{
-		if(IncrementallySkippedFiles.Num() != 0)
-		{
-			const FString FileList = FString::JoinBy(IncrementallySkippedFiles, TEXT("\n"), [](const FSourceControlStateRef& InFile) { return InFile->GetFilename(); });
-			UE_LOG(LogValidators, Log, TEXT("[%s] Skipping Files because they were already validated in a previous execution:\n%s"), *GetValidatorName(), *FileList);
-		}
-
-		if(!Validate(ChangelistService->GetCLDescription(), FilteredFiles, TagService->GetTagsArray()))
-		{
-			ValidationFinished(false);
-		}
-	}
 }
 
 
diff --git a/Engine/Source/Programs/SubmitTool/Private/Logic/Validators/ValidatorRunExecutable.cpp b/Engine/Source/Programs/SubmitTool/Private/Logic/Validators/ValidatorRunExecutable.cpp
--- a/Engine/Source/Programs/SubmitTool/Private/Logic/Validators/ValidatorRunExecutable.cpp
+++ b/Engine/Source/Programs/SubmitTool/Private/Logic/Validators/ValidatorRunExecutable.cpp
@@ -106,7 +106,7 @@ bool FValidatorRunExecutable::Activate()
 	{
 		TArray<FString> ExecutablePaths;
 		TypedDefinition->ExecutableCandidates.GenerateValueArray(ExecutablePaths);
-		UE_LOG(LogSubmitTool, Error, TEXT("[%s] None of the executable candidates exists locally:\n%s"), *ValidatorName, *FString::Join(ExecutablePaths, TEXT("\n")));
+		ActivationErrors.Add(FString::Printf(TEXT("[%s] None of the executable candidates exists locally:\n%s"), *ValidatorName, *FString::Join(ExecutablePaths, TEXT("\n"))));
 		bIsValidSetup = false;
 	}
 
@@ -136,14 +136,14 @@ bool FValidatorRunExecutable::Activate()
 		{
 			if(!ValidExtensions.Contains(FPaths::GetExtension(Paths.Value, true)))
 			{
-				UE_LOG(LogSubmitTool, Error, TEXT("Task '%s' executable has an invalid extension for this platform: %s"), *ValidatorName, *Paths.Value);
+				ActivationErrors.Add(FString::Printf(TEXT("Task '%s' executable has an invalid extension for this platform: %s"), *ValidatorName, *Paths.Value));
 				bIsValidSetup = false;
 			}
 		}
 
 		if(TypedDefinition->ExecutableCandidates.Num() == 0)
 		{
-			UE_LOG(LogSubmitTool, Error, TEXT("Task '%s' does not have a value for 'ExecutablePath' or 'ExecutableCandidates'."), *ValidatorName);
+			ActivationErrors.Add(FString::Printf(TEXT("Task '%s' does not have a value for 'ExecutablePath' or 'ExecutableCandidates'."), *ValidatorName));
 			bIsValidSetup = false;
 		}
 	}
@@ -151,20 +151,20 @@ bool FValidatorRunExecutable::Activate()
 	{
 		if(!ValidExtensions.Contains(FPaths::GetExtension(TypedDefinition->ExecutablePath, true)))
 		{
-			UE_LOG(LogSubmitTool, Error, TEXT("Task '%s' executable has an invalid extension for this platform: %s"), *ValidatorName, *TypedDefinition->ExecutablePath);
+			ActivationErrors.Add(FString::Printf(TEXT("Task '%s' executable has an invalid extension for this platform: %s"), *ValidatorName, *TypedDefinition->ExecutablePath));
 			bIsValidSetup = false;
 		}
 
 		if(!FPaths::FileExists(TypedDefinition->ExecutablePath))
 		{
-			UE_LOG(LogSubmitTool, Error, TEXT("Task '%s' executable is not found on disk: %s."), *ValidatorName, *TypedDefinition->ExecutablePath);
+			ActivationErrors.Add(FString::Printf(TEXT("Task '%s' executable is not found on disk: %s."), *ValidatorName, *TypedDefinition->ExecutablePath));
 			bIsValidSetup = false;
 		}
 	}
 
 	if(!TypedDefinition->ExecutablePath.IsEmpty() && TypedDefinition->ExecutableCandidates.Num() > 0)
 	{
-		UE_LOG(LogSubmitTool, Error, TEXT("Specifying ExecutablePath and ExecutableCandidates for task %s is not supported, please check your config."), *GetValidatorName());
+		ActivationErrors.Add(FString::Printf(TEXT("Specifying ExecutablePath and ExecutableCandidates for task %s is not supported, please check your config."), *GetValidatorName()));
 	}
 
 
