Manual cancellation mark validators as failed
Performance improvements when CL contains many files

diff --git a/Engine/Source/Programs/SubmitTool/Private/Logic/TasksService.cpp b/Engine/Source/Programs/SubmitTool/Private/Logic/TasksService.cpp
--- a/Engine/Source/Programs/SubmitTool/Private/Logic/TasksService.cpp
+++ b/Engine/Source/Programs/SubmitTool/Private/Logic/TasksService.cpp
@@ -37,11 +37,11 @@ void FTasksService::InitializeTasks(const TArray<TSharedRef<FValidatorBase>>& In
 	{
 		if(!Task->Activate())
 		{
-			UE_LOG(LogSubmitTool, Error, TEXT("[%s] has errors and is in an invalid state."), *Task->GetValidatorNameId().ToString());
+			UE_LOG(LogSubmitToolDebug, Error, TEXT("[%s] has errors and is in an invalid state."), *Task->GetValidatorNameId().ToString());
 		}
 		else
 		{
-			UE_LOG(LogValidators, Log, TEXT("Task '%s' is active."), *Task->GetValidatorNameId().ToString());
+			UE_LOG(LogSubmitToolDebug, Log, TEXT("Task '%s' is active."), *Task->GetValidatorNameId().ToString());
 		}
 
 		// Create completed handle and register it with ValidationService
@@ -469,13 +469,13 @@ const TArray<FString> FTasksService::GetAddendums() const
 }
 
 
-void FTasksService::StopTasks(const FName& InTaskId)
+void FTasksService::StopTasks(const FName& InTaskId, bool InbAsFailed)
 {
 	for(const TPair<FName, TSharedPtr<FValidatorBase>>& TaskPair : Tasks)
 	{
 		if(InTaskId.IsNone() || InTaskId.IsEqual(TaskPair.Value->GetValidatorNameId()))
 		{
-			TaskPair.Value->CancelValidation();
+			TaskPair.Value->CancelValidation(InbAsFailed);
 		}
 	}
 }
diff --git a/Engine/Source/Programs/SubmitTool/Private/Logic/TasksService.h b/Engine/Source/Programs/SubmitTool/Private/Logic/TasksService.h
--- a/Engine/Source/Programs/SubmitTool/Private/Logic/TasksService.h
+++ b/Engine/Source/Programs/SubmitTool/Private/Logic/TasksService.h
@@ -44,7 +44,7 @@ public:
 
 	virtual const TArray<FString> GetAddendums() const;
 
-	virtual void StopTasks(const FName& InValidatorId = FName());
+	virtual void StopTasks(const FName& InValidatorId = FName(), bool InbAsFailed = false);
 	virtual void StopTasksByArea(const ETaskArea& InArea);
 
 	FOnTaskFinished OnTasksQueueFinished;
diff --git a/Engine/Source/Programs/SubmitTool/Private/Logic/Validators/ValidatorBase.cpp b/Engine/Source/Programs/SubmitTool/Private/Logic/Validators/ValidatorBase.cpp
--- a/Engine/Source/Programs/SubmitTool/Private/Logic/Validators/ValidatorBase.cpp
+++ b/Engine/Source/Programs/SubmitTool/Private/Logic/Validators/ValidatorBase.cpp
@@ -155,7 +155,7 @@ bool FValidatorBase::Activate()
 
 void FValidatorBase::InvalidateLocalFileModifications()
 {
-	if((Definition->TaskArea & ETaskArea::LocalFiles) == ETaskArea::LocalFiles && (State == EValidationStates::Valid || State == EValidationStates::Running || State == EValidationStates::Skipped))
+	if((Definition->TaskArea & ETaskArea::LocalFiles) == ETaskArea::LocalFiles && (State == EValidationStates::Valid || State == EValidationStates::Running || State == EValidationStates::Skipped || State == EValidationStates::Not_Applicable))
 	{
 		FFileManagerGeneric FileManager;
 		for(const FSourceControlStateRef& File : ServiceProvider.Pin()->GetService<FChangelistService>()->GetFilesInCL())
@@ -192,13 +192,8 @@ const FString FValidatorBase::GetStatusText() const
 	const FString StateStr = StaticEnum<EValidationStates>()->GetNameStringByValue(static_cast<int64>(State))
 		.Replace(TEXT("_"), TEXT(" "));
 
-	if(State == EValidationStates::Skipped || State == EValidationStates::Not_Run)
+	if(State == EValidationStates::Skipped || State == EValidationStates::Not_Run || State == EValidationStates::Not_Applicable)
 	{
-		if (State == EValidationStates::Skipped && !IsRelevantToCL())
-		{
-			return TEXT("Not Applicable");
-		}
-
 		return StateStr;
 	}
 
diff --git a/Engine/Source/Programs/SubmitTool/Private/Logic/Validators/ValidatorBase.h b/Engine/Source/Programs/SubmitTool/Private/Logic/Validators/ValidatorBase.h
--- a/Engine/Source/Programs/SubmitTool/Private/Logic/Validators/ValidatorBase.h
+++ b/Engine/Source/Programs/SubmitTool/Private/Logic/Validators/ValidatorBase.h
@@ -40,6 +40,7 @@ enum class EValidationStates : uint8
 {
 	Not_Run,
 	Skipped,
+	Not_Applicable,
 	Running,
 	Valid,
 	Failed,
@@ -91,7 +92,7 @@ public:
 
 	bool GetHasPassed() const
 	{
-		return State == EValidationStates::Valid || State == EValidationStates::Skipped;
+		return State == EValidationStates::Valid || State == EValidationStates::Skipped || State == EValidationStates::Not_Applicable;
 	}
 
 	bool GetIsRunningOrQueued() const
@@ -121,14 +122,14 @@ public:
 
 	virtual void Tick(float InDeltaTime);
 
-	virtual void CancelValidation()
+	virtual void CancelValidation(bool InbAsFailed = false)
 	{
 		if(State == EValidationStates::Running)
 		{
 			StopInternalValidations();
 		}
 
-		State = EValidationStates::Not_Run;
+		State = InbAsFailed ? EValidationStates::Failed : EValidationStates::Not_Run;
 	}
 
 	virtual void SetQueued(bool InbForceRun = false)
@@ -144,7 +145,7 @@ public:
 
 	void SetNotApplicable()
 	{
-		State = EValidationStates::Skipped;
+		State = EValidationStates::Not_Applicable;
 	}
 
 	virtual bool Activate();
@@ -200,6 +201,7 @@ protected:
 	FValidatorOptionsProvider OptionsProvider;
 	TWeakPtr<FSubmitToolServiceProvider> ServiceProvider;
 	TArray<FSourceControlStateRef> FilteredFiles;
+	TArray<FString> ActivationErrors;
 	mutable TArray<FString> ErrorListCache;
 	bool bIsValidSetup = false;
 	bool bForceRun = false;
diff --git a/Engine/Source/Programs/SubmitTool/Private/Models/ModelInterface.h b/Engine/Source/Programs/SubmitTool/Private/Models/ModelInterface.h
--- a/Engine/Source/Programs/SubmitTool/Private/Models/ModelInterface.h
+++ b/Engine/Source/Programs/SubmitTool/Private/Models/ModelInterface.h
@@ -137,7 +137,7 @@ public:
 	bool IsValidationRunning() const												{ return ValidationService->GetIsAnyTaskRunning(); }
 	const TArray<TWeakPtr<const FValidatorBase>>& GetValidators() const				{ return ValidationService->GetTasks(); }
 	const TArray<TWeakPtr<const FValidatorBase>>& GetPreSubmitOperations() const	{ return PresubmitOperationsService->GetTasks(); }
-	void CancelValidations(const FName& InValidatorId = FName()) const				{ ValidationService->StopTasks(InValidatorId); }
+	void CancelValidations(const FName& InValidatorId = FName(), bool InbAsFailed = false) const { ValidationService->StopTasks(InValidatorId, InbAsFailed); }
 	void CheckForFileEdits() const													{ ValidationService->CheckForLocalFileEdit(); }
 
 	FDelegateHandle AddSingleValidatorFinishedCallback
diff --git a/Engine/Source/Programs/SubmitTool/Private/View/Widgets/ValidatorsWidget.cpp b/Engine/Source/Programs/SubmitTool/Private/View/Widgets/ValidatorsWidget.cpp
--- a/Engine/Source/Programs/SubmitTool/Private/View/Widgets/ValidatorsWidget.cpp
+++ b/Engine/Source/Programs/SubmitTool/Private/View/Widgets/ValidatorsWidget.cpp
@@ -339,6 +339,7 @@ public:
 						break;
 					case EValidationStates::Valid:
 					case EValidationStates::Skipped:
+					case EValidationStates::Not_Applicable:
 						return FAppStyle::GetColor("ValidatorStateSuccess");
 						break;
 					case EValidationStates::Running:
@@ -384,7 +385,7 @@ public:
 					.OnClicked_Lambda([this]()
 					{
 						if(Node.Pin()->GetIsRunningOrQueued()){
-							ModelInterface->CancelValidations(Node.Pin()->GetValidatorNameId());
+							ModelInterface->CancelValidations(Node.Pin()->GetValidatorNameId(), true);
 						}else{
 							ModelInterface->ValidateSingle(Node.Pin()->GetValidatorNameId(), false);
 						}
@@ -401,7 +402,7 @@ public:
 				.Padding(3, 0, 3, 0)
 				[
 					SNew(SButton)
-					.IsEnabled_Lambda([this]() {return Node.Pin()->GetValidatorState() != EValidationStates::Not_Run && Node.Pin()->GetValidatorState() != EValidationStates::Skipped; })
+					.IsEnabled_Lambda([this]() { return Node.Pin()->GetValidatorState() != EValidationStates::Not_Run && Node.Pin()->GetValidatorState() != EValidationStates::Skipped && Node.Pin()->GetValidatorState() != EValidationStates::Not_Applicable; })
 					.HAlign(EHorizontalAlignment::HAlign_Center)
 					.VAlign(EVerticalAlignment::VAlign_Center)
 					.OnClicked_Lambda([this]() {ViewLogCallback.ExecuteIfBound(Node.Pin()); return FReply::Handled(); })					
