Open Jira Tickets on submit

diff --git a/Engine/Source/Programs/SubmitTool/Private/Logic/ChangelistService.cpp b/Engine/Source/Programs/SubmitTool/Private/Logic/ChangelistService.cpp
--- a/Engine/Source/Programs/SubmitTool/Private/Logic/ChangelistService.cpp
+++ b/Engine/Source/Programs/SubmitTool/Private/Logic/ChangelistService.cpp
@@ -8,6 +8,7 @@
 #include "Models/ModelInterface.h"
 #include "Telemetry/TelemetryService.h"
 #include "CommandLine/CmdLineParameters.h"
+#include "HAL/PlatformApplicationMisc.h"
 
 #include "ISourceControlProvider.h"
 #include "ISourceControlOperation.h"
@@ -337,6 +338,13 @@ void FChangelistService::Submit(const FString& InDescriptionAddendum, const FSou
 				TSharedRef<FCheckIn> CheckIn = StaticCastSharedRef<FCheckIn>(Operation);
 				UE_LOG(LogSubmitToolP4, Log, TEXT("Submit operation succeeded: %s"), *CheckIn->GetSuccessMessage().ToString());
 
+				const FString CLString = CheckIn->GetSuccessMessage().ToString().Replace(TEXT("Submitted changelist "), TEXT(""), ESearchCase::IgnoreCase).TrimStartAndEnd();
+				if (CLString.IsNumeric())
+				{
+					FPlatformApplicationMisc::ClipboardCopy(*CLString);
+					UE_LOG(LogSubmitToolP4, Log, TEXT("Submitted CL copied to clipboard: %s"), *CLString);
+				}
+
 				FTelemetryService::Get()->SubmitSucceeded();
 			}
 			else if(Result == ECommandResult::Cancelled)
