New ST Version 5.4.0

diff --git a/Engine/Source/Programs/SubmitTool/Private/Version/AppVersionDefines.h b/Engine/Source/Programs/SubmitTool/Private/Version/AppVersionDefines.h
--- a/Engine/Source/Programs/SubmitTool/Private/Version/AppVersionDefines.h
+++ b/Engine/Source/Programs/SubmitTool/Private/Version/AppVersionDefines.h
@@ -20,8 +20,8 @@
 #endif
 
 #define SUBMIT_TOOL_VERSION_MAJOR 5
-#define SUBMIT_TOOL_VERSION_MINOR 3
-#define SUBMIT_TOOL_VERSION_PATCH 1
+#define SUBMIT_TOOL_VERSION_MINOR 4
+#define SUBMIT_TOOL_VERSION_PATCH 0
 
 #ifndef SUBMIT_TOOL_CHANGELIST
 #define SUBMIT_TOOL_CHANGELIST_STRING TEXT("NoCL")
