Adding telemetry to start preflights

diff --git a/Engine/Source/Programs/SubmitTool/Private/Logic/PreflightService.cpp b/Engine/Source/Programs/SubmitTool/Private/Logic/PreflightService.cpp
--- a/Engine/Source/Programs/SubmitTool/Private/Logic/PreflightService.cpp
+++ b/Engine/Source/Programs/SubmitTool/Private/Logic/PreflightService.cpp
@@ -18,6 +18,7 @@
 #include "Logic/TagService.h"
 #include "Logging/SubmitToolLog.h"
 #include "Models/ModelInterface.h"
+#include "Telemetry/TelemetryService.h"
 #include "Parameters/SubmitToolParameters.h"
 
 FPreflightService::FPreflightService(
@@ -562,6 +563,11 @@ void FPreflightService::StartPreflight()
 	{
 		UE_LOG(LogSubmitTool, Log, TEXT("Preflight: Starting preflight with URL: \"%s\""), *StartPreflightUrl);
 
+		FTelemetryService::Get()->CustomEvent(TEXT("SubmitTool.PreflightLaunched"), MakeAnalyticsEventAttributeArray(
+			TEXT("PreflightURL"), StartPreflightUrl,
+			TEXT("Stream"), StreamName
+		));
+
 		FPlatformProcess::LaunchURL(*StartPreflightUrl, nullptr, nullptr);
 		State = EPreflightServiceState::Idle;
 
