Fix for shelved filelist not being emptied on forced refresh
Selecting the right preflights for UE streams

diff --git a/Engine/Source/Programs/SubmitTool/Private/Logic/ChangelistService.cpp b/Engine/Source/Programs/SubmitTool/Private/Logic/ChangelistService.cpp
--- a/Engine/Source/Programs/SubmitTool/Private/Logic/ChangelistService.cpp
+++ b/Engine/Source/Programs/SubmitTool/Private/Logic/ChangelistService.cpp
@@ -657,6 +657,7 @@ const TArray<FString>& FChangelistService::GetFilesDepotPaths(bool bForce)
 	FScopeLock Lock(&Mutex);
 	if(bForce || (FilesDepotPaths.IsEmpty() && !FilesInCL.IsEmpty()))
 	{
+		FilesDepotPaths.Reset();
 		UE::Tasks::TTask<bool> OpenedTask = SourceControlService->RunCommand(TEXT("opened"), { "-c", CLID }, FOnSCCCommandComplete::CreateLambda(
 			[this](bool bSuccess, const FSCCRecordSet& InResultValues, const FSourceControlResultInfo& InResultsInfo)
 			{
@@ -700,6 +701,7 @@ const TArray<FString>& FChangelistService::GetShelvedFilesDepotPaths(bool bForce
 	FScopeLock Lock(&Mutex);
 	if(bForce || (ShelvedFilesDepotPaths.IsEmpty() && !ShelvedFilesInCL.IsEmpty()))
 	{
+		ShelvedFilesDepotPaths.Reset();
 		UE::Tasks::TTask<bool> DescribeTask = SourceControlService->RunCommand(TEXT("describe"), { "-S", CLID }, FOnSCCCommandComplete::CreateLambda(
 			[this](bool bSuccess, const FSCCRecordSet& InResultValues, const FSourceControlResultInfo& InResultsInfo)
 			{
