Adding opening Jira ticket & copy submitted CL as an option toggle
Fixed dangling task when OIDCToken process never finishes
Added shelve-local conflight message warning

diff --git a/Engine/Source/Programs/SubmitTool/Private/Logic/ChangelistService.cpp b/Engine/Source/Programs/SubmitTool/Private/Logic/ChangelistService.cpp
--- a/Engine/Source/Programs/SubmitTool/Private/Logic/ChangelistService.cpp
+++ b/Engine/Source/Programs/SubmitTool/Private/Logic/ChangelistService.cpp
@@ -8,7 +8,6 @@
 #include "Models/ModelInterface.h"
 #include "Telemetry/TelemetryService.h"
 #include "CommandLine/CmdLineParameters.h"
-#include "HAL/PlatformApplicationMisc.h"
 
 #include "ISourceControlProvider.h"
 #include "ISourceControlOperation.h"
@@ -337,14 +336,6 @@ void FChangelistService::Submit(const FString& InDescriptionAddendum, const FSou
 			{
 				TSharedRef<FCheckIn> CheckIn = StaticCastSharedRef<FCheckIn>(Operation);
 				UE_LOG(LogSubmitToolP4, Log, TEXT("Submit operation succeeded: %s"), *CheckIn->GetSuccessMessage().ToString());
-
-				const FString CLString = CheckIn->GetSuccessMessage().ToString().Replace(TEXT("Submitted changelist "), TEXT(""), ESearchCase::IgnoreCase).TrimStartAndEnd();
-				if (CLString.IsNumeric())
-				{
-					FPlatformApplicationMisc::ClipboardCopy(*CLString);
-					UE_LOG(LogSubmitToolP4, Log, TEXT("Submitted CL copied to clipboard: %s"), *CLString);
-				}
-
 				FTelemetryService::Get()->SubmitSucceeded();
 			}
 			else if(Result == ECommandResult::Cancelled)
diff --git a/Engine/Source/Programs/SubmitTool/Private/Models/ModelInterface.cpp b/Engine/Source/Programs/SubmitTool/Private/Models/ModelInterface.cpp
--- a/Engine/Source/Programs/SubmitTool/Private/Models/ModelInterface.cpp
+++ b/Engine/Source/Programs/SubmitTool/Private/Models/ModelInterface.cpp
@@ -9,6 +9,8 @@
 #include "Widgets/Docking/SDockTab.h"
 #include "Containers/Ticker.h"
 #include "Configuration/Configuration.h"
+#include "SourceControlOperations.h"
+#include "HAL/PlatformApplicationMisc.h"
 
 #include "Logic/Validators/PreflightValidator.h"
 #include "Logic/Validators/TagValidator.h"
@@ -38,7 +40,7 @@ FModelInterface::FModelInterface(const FSubmitToolParameters& InParameters) :
 
 	ServiceProvider = MakeShared<FSubmitToolServiceProvider>();
 	// Initialize services
-	if(Parameters.GeneralParameters.CacheFile.IsEmpty())
+	if (Parameters.GeneralParameters.CacheFile.IsEmpty())
 	{
 		ServiceProvider->RegisterService<ICacheDataService>(MakeShared<FNoOpCacheDataService>());
 	}
@@ -76,58 +78,58 @@ FModelInterface::FModelInterface(const FSubmitToolParameters& InParameters) :
 	ParsePreSubmitOperations();
 
 	OnValidationStateUpdatedHandle = ValidationService->OnTasksRunResultUpdated.Add(FOnTaskRunStateChanged::FDelegate::CreateLambda([this](bool bIsValid) {
-		if(bIsValid)
+		if (bIsValid)
 		{
 			bool bOptionalFailures = false;
-			for(const TWeakPtr<const FValidatorBase>& Validator : ValidationService->GetTasks())
+			for (const TWeakPtr<const FValidatorBase>& Validator : ValidationService->GetTasks())
 			{
-				if(!Validator.Pin()->GetIsRunningOrQueued() && !Validator.Pin()->GetHasPassed())
+				if (!Validator.Pin()->GetIsRunningOrQueued() && !Validator.Pin()->GetHasPassed())
 				{
 					bOptionalFailures = true;
 				}
 			}
 
 			UE_LOG(LogSubmitTool, Log, TEXT("The required local validation has succeeded, you're ALLOWED TO SUBMIT."))
-			if(ValidationService->GetIsAnyTaskRunning())
-			{
-				UE_LOG(LogSubmitTool, Warning, TEXT("You still have optional validations running you might want to consider waiting for them to finish."))
-			}
-			if(bOptionalFailures)
+				if (ValidationService->GetIsAnyTaskRunning())
+				{
+					UE_LOG(LogSubmitTool, Warning, TEXT("You still have optional validations running you might want to consider waiting for them to finish."))
+				}
+			if (bOptionalFailures)
 			{
 				UE_LOG(LogSubmitTool, Warning, TEXT("You have optional validations that have failed, you can still proceed with the submission if you consider that these failures are not relevant. Please make sure this is the case."))
 			}
 		}
-	}));
-	
+		}));
+
 	OnSingleValidationFinishedHandle = ValidationService->OnSingleTaskFinished.AddLambda([this](const FValidatorBase& InTask) {
-			ReevaluateSubmitToolTag();
+		ReevaluateSubmitToolTag();
 
-			if(bPreflightQueued && CanLaunchPreflight())
-			{
-				bPreflightQueued = false;
-				PreflightService->RequestPreflight();
-			}
+		if (bPreflightQueued && CanLaunchPreflight())
+		{
+			bPreflightQueued = false;
+			PreflightService->RequestPreflight();
+		}
 		});
 
 	OnValidationFinishedHandle = ValidationService->OnTasksQueueFinished.Add(FOnTaskFinished::FDelegate::CreateLambda([this](bool bIsValid)
 		{
 			MainTab.Pin()->GetParentWindow()->DrawAttention(FWindowDrawAttentionParameters());
-			
-			if(bIsValid)
+
+			if (bIsValid)
 			{
 				if (!bPreflightQueued && bSubmitOnSuccessfulValidation && !IsIntegrationRequired())
 				{
 					bool bAllSucceedIncludingOptional = true;
-					for(const TWeakPtr<const FValidatorBase>& Validator : ValidationService->GetTasks())
+					for (const TWeakPtr<const FValidatorBase>& Validator : ValidationService->GetTasks())
 					{
-						if(Validator.IsValid() && !Validator.Pin()->GetHasPassed())
+						if (Validator.IsValid() && !Validator.Pin()->GetHasPassed())
 						{
 							bAllSucceedIncludingOptional = false;
 						}
 					}
 
 
-					if(bAllSucceedIncludingOptional)
+					if (bAllSucceedIncludingOptional)
 					{
 						UE_LOG(LogSubmitTool, Log, TEXT("User has opted in to automatically submit on validation successful. Proceeding with submission..."));
 						StartSubmitProcess();
@@ -164,15 +166,15 @@ void FModelInterface::Dispose() const
 {
 	ChangelistService->CancelP4Operations();
 
-	if(GetInputEnabled())
+	if (GetInputEnabled())
 	{
 		ChangelistService->SendCLDescriptionToP4(EConcurrency::Synchronous);
 	}
 
-	for(const TPair<FString, TSharedPtr<FIntegrationOptionBase>>& IntegrationOption : ServiceProvider->GetService<FIntegrationService>()->GetIntegrationOptions())
+	for (const TPair<FString, TSharedPtr<FIntegrationOptionBase>>& IntegrationOption : ServiceProvider->GetService<FIntegrationService>()->GetIntegrationOptions())
 	{
 		FString Value;
-		if(IntegrationOption.Value->GetJiraValue(Value) && !Value.IsEmpty())
+		if (IntegrationOption.Value->GetJiraValue(Value) && !Value.IsEmpty())
 		{
 			ServiceProvider->GetService<ICacheDataService>()->SetIntegrationFieldValue(GetCLID(), IntegrationOption.Key, Value);
 		}
@@ -186,20 +188,20 @@ void FModelInterface::ParseValidators() const
 {
 	TArray<TSharedRef<FValidatorBase>> Tasks;
 
-	for(const TPair<FName, FString> DefinitionPair : Parameters.Validators)
+	for (const TPair<FName, FString> DefinitionPair : Parameters.Validators)
 	{
 		FValidatorDefinition TaskDefinition;
 		FStringOutputDevice Errors;
 		FValidatorDefinition::StaticStruct()->ImportText(*DefinitionPair.Value, &TaskDefinition, nullptr, 0, &Errors, FValidatorDefinition::StaticStruct()->GetName());
 
-		if(!Errors.IsEmpty())
+		if (!Errors.IsEmpty())
 		{
 			UE_LOG(LogSubmitTool, Error, TEXT("Error loading parameter file %s"), *Errors);
 			FModelInterface::SetErrorState();
 			continue;
 		}
 
-		if(TaskDefinition.Type.TrimStartAndEnd().IsEmpty())
+		if (TaskDefinition.Type.TrimStartAndEnd().IsEmpty())
 		{
 			UE_LOG(LogSubmitTool, Error, TEXT("Task %s didn't have a Type."), *DefinitionPair.Key.ToString());
 			continue;
@@ -211,23 +213,23 @@ void FModelInterface::ParseValidators() const
 			continue;
 		}
 
-		if(TaskDefinition.Type.Equals(SubmitToolParseConstants::TagValidator, ESearchCase::IgnoreCase))
+		if (TaskDefinition.Type.Equals(SubmitToolParseConstants::TagValidator, ESearchCase::IgnoreCase))
 		{
 			Tasks.Add(MakeShared<FTagValidator>(DefinitionPair.Key, Parameters, ServiceProvider.ToSharedRef(), DefinitionPair.Value));
 		}
-		else if(TaskDefinition.Type.Equals(SubmitToolParseConstants::UBTValidator, ESearchCase::IgnoreCase))
+		else if (TaskDefinition.Type.Equals(SubmitToolParseConstants::UBTValidator, ESearchCase::IgnoreCase))
 		{
 			Tasks.Add(MakeShared<FUBTValidator>(DefinitionPair.Key, Parameters, ServiceProvider.ToSharedRef(), DefinitionPair.Value));
 		}
-		else if(TaskDefinition.Type.Equals(SubmitToolParseConstants::CustomValidator, ESearchCase::IgnoreCase))
+		else if (TaskDefinition.Type.Equals(SubmitToolParseConstants::CustomValidator, ESearchCase::IgnoreCase))
 		{
 			Tasks.Add(MakeShared<FValidatorRunExecutable>(DefinitionPair.Key, Parameters, ServiceProvider.ToSharedRef(), DefinitionPair.Value));
 		}
-		else if(TaskDefinition.Type.Equals(SubmitToolParseConstants::CrossChangelistValidator, ESearchCase::IgnoreCase))
+		else if (TaskDefinition.Type.Equals(SubmitToolParseConstants::CrossChangelistValidator, ESearchCase::IgnoreCase))
 		{
 			Tasks.Add(MakeShared<FCrossChangelistValidator>(DefinitionPair.Key, Parameters, ServiceProvider.ToSharedRef(), DefinitionPair.Value));
 		}
-		else if(TaskDefinition.Type.Equals(SubmitToolParseConstants::PreflightValidator, ESearchCase::IgnoreCase))
+		else if (TaskDefinition.Type.Equals(SubmitToolParseConstants::PreflightValidator, ESearchCase::IgnoreCase))
 		{
 			Tasks.Add(MakeShared<FPreflightValidator>(DefinitionPair.Key, Parameters, ServiceProvider.ToSharedRef(), DefinitionPair.Value));
 		}
@@ -257,30 +259,30 @@ void FModelInterface::ParsePreSubmitOperations() const
 {
 	TArray<TSharedRef<FValidatorBase>> Tasks;
 
-	for(const TPair<FName, FString> DefinitionPair : Parameters.PresubmitOperations)
+	for (const TPair<FName, FString> DefinitionPair : Parameters.PresubmitOperations)
 	{
 		FValidatorDefinition TaskDefinition;
 		FStringOutputDevice Errors;
 		FValidatorDefinition::StaticStruct()->ImportText(*DefinitionPair.Value, &TaskDefinition, nullptr, 0, &Errors, FValidatorDefinition::StaticStruct()->GetName());
 
-		if(!Errors.IsEmpty())
+		if (!Errors.IsEmpty())
 		{
 			UE_LOG(LogSubmitTool, Error, TEXT("Error loading parameter file %s"), *Errors);
 			FModelInterface::SetErrorState();
 			continue;
 		}
 
-		if(TaskDefinition.Type.TrimStartAndEnd().IsEmpty())
+		if (TaskDefinition.Type.TrimStartAndEnd().IsEmpty())
 		{
 			UE_LOG(LogSubmitTool, Error, TEXT("Task %s didn't have a Type."), *DefinitionPair.Key.ToString());
 			continue;
 		}
 
-		if(TaskDefinition.Type.Equals(SubmitToolParseConstants::CustomValidator, ESearchCase::IgnoreCase))
+		if (TaskDefinition.Type.Equals(SubmitToolParseConstants::CustomValidator, ESearchCase::IgnoreCase))
 		{
 			Tasks.Add(MakeShared<FValidatorRunExecutable>(DefinitionPair.Key, Parameters, ServiceProvider.ToSharedRef(), DefinitionPair.Value));
 		}
-		else if(TaskDefinition.Type.Equals(SubmitToolParseConstants::VirtualizationToolOp, ESearchCase::IgnoreCase))
+		else if (TaskDefinition.Type.Equals(SubmitToolParseConstants::VirtualizationToolOp, ESearchCase::IgnoreCase))
 		{
 			Tasks.Add(MakeShared<FVirtualizationOperation>(DefinitionPair.Key, Parameters, ServiceProvider.ToSharedRef(), DefinitionPair.Value));
 		}
@@ -297,11 +299,11 @@ void FModelInterface::ParsePreSubmitOperations() const
 
 void FModelInterface::SetCLDescription(const FText& newDescription, bool DoNotInvalidate) const
 {
-	if(ChangelistService->SetCLDescription(newDescription.ToString()))
+	if (ChangelistService->SetCLDescription(newDescription.ToString()))
 	{
 		TagService->ParseCLDescription();
 
-		if(!DoNotInvalidate)
+		if (!DoNotInvalidate)
 		{
 			ValidationService->InvalidateForChanges(ETaskArea::Changelist);
 		}
@@ -310,9 +312,9 @@ void FModelInterface::SetCLDescription(const FText& newDescription, bool DoNotIn
 
 void FModelInterface::SendDescriptionToP4() const
 {
-	if(GetInputEnabled())
+	if (GetInputEnabled())
 	{
-		if(IsP4OperationRunning())
+		if (IsP4OperationRunning())
 		{
 			UE_LOG(LogSubmitToolP4, Log, TEXT("Attempted to send description to P4, but another operation is already running"));
 			return;
@@ -326,14 +328,14 @@ bool FModelInterface::CanLaunchPreflight() const
 {
 	// Check Validators which are validating files, ignore changelist (description, valid tags) validators when we evaluate if we 
 	// allow the user to trigger a preflight
-	for(const TWeakPtr<const FValidatorBase>& Validator : ValidationService->GetTasks())
+	for (const TWeakPtr<const FValidatorBase>& Validator : ValidationService->GetTasks())
 	{
-		if(Validator.IsValid())
+		if (Validator.IsValid())
 		{
 			const TSharedPtr<const FValidatorBase> Pinned = Validator.Pin();
-			if((Pinned->Definition->TaskArea & ETaskArea::ShelveAndLocalFiles) != ETaskArea::None)
+			if ((Pinned->Definition->TaskArea & ETaskArea::ShelveAndLocalFiles) != ETaskArea::None)
 			{
-				if((Pinned->Definition->IsRequired && !Pinned->GetHasPassed()) || (!Pinned->Definition->IsRequired && Pinned->GetIsRunningOrQueued()))
+				if ((Pinned->Definition->IsRequired && !Pinned->GetHasPassed()) || (!Pinned->Definition->IsRequired && Pinned->GetIsRunningOrQueued()))
 				{
 					return false;
 				}
@@ -346,7 +348,7 @@ bool FModelInterface::CanLaunchPreflight() const
 
 void FModelInterface::ReevaluateSubmitToolTag()
 {
-	UpdateSubmitToolTag(ValidationService->GetIsRunSuccessful(!IsIntegrationRequired()));	
+	UpdateSubmitToolTag(ValidationService->GetIsRunSuccessful(!IsIntegrationRequired()));
 }
 
 void FModelInterface::UpdateSubmitToolTag(bool InbAdd)
@@ -355,25 +357,25 @@ void FModelInterface::UpdateSubmitToolTag(bool InbAdd)
 	FString SubmitToolTag = FString::Format(TEXT("#submittool {0}\n"), { FAppVersion::GetVersion() });
 	FString DescriptionCopy = ChangelistService->GetCLDescription();
 
-	if(InbAdd)
+	if (InbAdd)
 	{
-		if(!HasSubmitToolTag())
+		if (!HasSubmitToolTag())
 		{
 			UE_LOG(LogSubmitToolDebug, Log, TEXT("Added Submit Tool tag"));
 			DescriptionCopy.Append(TEXT("\n") + SubmitToolTag);
 			ChangelistService->SetCLDescription(DescriptionCopy, true);
 		}
 	}
-	else if(HasSubmitToolTag())
+	else if (HasSubmitToolTag())
 	{
 		FString VersionlessTag = TEXT("#submittool ");
 		int32 Loc = DescriptionCopy.Find(VersionlessTag);
-		if(Loc >= 0)
+		if (Loc >= 0)
 		{
 			size_t EndPos = Loc + VersionlessTag.Len();
-			while(EndPos < DescriptionCopy.Len())
+			while (EndPos < DescriptionCopy.Len())
 			{
-				if(DescriptionCopy[EndPos] == TCHAR('\n'))
+				if (DescriptionCopy[EndPos] == TCHAR('\n'))
 				{
 					++EndPos;
 					break;
@@ -381,7 +383,7 @@ void FModelInterface::UpdateSubmitToolTag(bool InbAdd)
 				++EndPos;
 			}
 
-			if(Loc != 0 && DescriptionCopy[Loc - 1] == '\n' && (EndPos - DescriptionCopy.Len()) < 2)
+			if (Loc != 0 && DescriptionCopy[Loc - 1] == '\n' && (EndPos - DescriptionCopy.Len()) < 2)
 			{
 				Loc--;
 			}
@@ -404,7 +406,7 @@ bool FModelInterface::HasSubmitToolTag() const
 
 void FModelInterface::UpdateCLFromP4Async() const
 {
-	if(GetInputEnabled() || SubmitToolState == ESubmitToolAppState::Errored || SubmitToolState == ESubmitToolAppState::SubmitLocked)
+	if (GetInputEnabled() || SubmitToolState == ESubmitToolAppState::Errored || SubmitToolState == ESubmitToolAppState::SubmitLocked)
 	{
 		ChangelistService->FetchChangelistDataAsync();
 	}
@@ -417,7 +419,7 @@ bool FModelInterface::GetInputEnabled()
 
 void FModelInterface::RequestPreflight()
 {
-	if(CanLaunchPreflight())
+	if (CanLaunchPreflight())
 	{
 		PreflightService->RequestPreflight();
 	}
@@ -472,32 +474,32 @@ void FModelInterface::StartSubmitProcess()
 
 	// Check if any last minute file changes have come in that invalidated any validators.
 	CheckForFileEdits();
-	if(IsCLValid())
+	if (IsCLValid())
 	{
-		if(PrepareSubmitCallBack.IsBound())
+		if (PrepareSubmitCallBack.IsBound())
 		{
 			PrepareSubmitCallBack.Broadcast();
 		}
 
 		UpdateSubmitToolTag(true);
 
-		if(Parameters.IncompatibleFilesParams.IncompatibleFileGroups.Num() > 0)
+		if (Parameters.IncompatibleFilesParams.IncompatibleFileGroups.Num() > 0)
 		{
 			const TArray<FString>& FilesInCL = ChangelistService->GetFilesDepotPaths();
 
 			FConfiguration::AddOrUpdateEntry(TEXT("$(StreamRoot)"), *GetRootStreamName());
 
-			for(const FIncompatibleFilesGroup& FileGroup : Parameters.IncompatibleFilesParams.IncompatibleFileGroups)
+			for (const FIncompatibleFilesGroup& FileGroup : Parameters.IncompatibleFilesParams.IncompatibleFileGroups)
 			{
 				TArray<size_t, TInlineAllocator<8>> Indexes;
 
-				for(const FString& File : FilesInCL)
+				for (const FString& File : FilesInCL)
 				{
-					for(size_t i = 0; i < FileGroup.FileGroups.Num(); ++i)
+					for (size_t i = 0; i < FileGroup.FileGroups.Num(); ++i)
 					{
-						if(File.Contains(FConfiguration::Substitute(FileGroup.FileGroups[i]), ESearchCase::IgnoreCase))
+						if (File.Contains(FConfiguration::Substitute(FileGroup.FileGroups[i]), ESearchCase::IgnoreCase))
 						{
-							if(!Indexes.Contains(i))
+							if (!Indexes.Contains(i))
 							{
 								Indexes.Add(i);
 							}
@@ -506,12 +508,12 @@ void FModelInterface::StartSubmitProcess()
 					}
 				}
 
-				if(Indexes.Num() > 1)
+				if (Indexes.Num() > 1)
 				{
 					const FText TextTitle = FText::FromString(FileGroup.Title);
 					const FText TextDescription = FText::FromString(FileGroup.GetMessage());
 
-					if(FileGroup.bIsError)
+					if (FileGroup.bIsError)
 					{
 						FDialogFactory::ShowInformationDialog(TextTitle, TextDescription);
 						UE_LOG(LogSubmitTool, Log, TEXT("Submission canceled due to incompatible files"));
@@ -519,25 +521,76 @@ void FModelInterface::StartSubmitProcess()
 					}
 					else
 					{
-						if(FDialogFactory::ShowConfirmDialog(TextTitle, TextDescription) != EDialogFactoryResult::Confirm)
+						if (FDialogFactory::ShowConfirmDialog(TextTitle, TextDescription) != EDialogFactoryResult::Confirm)
 						{
 							UE_LOG(LogSubmitTool, Log, TEXT("Submission canceled by user"));
 							return;
 						}
 					}
-				}				
+				}
 			}
 		}
 
-		ChangeState(ESubmitToolAppState::Submitting, SubmitToolState == ESubmitToolAppState::SubmitLocked && bIsUserInAllowlist);
-		
-
-		if(HasShelvedFiles())
+		if (HasShelvedFiles())
 		{
-			DeleteShelvedFiles();
+			const TArray<FString>& ShelvedFiles = ChangelistService->GetShelvedFilesDepotPaths();
+			const TArray<FString>& LocalFiles = ChangelistService->GetFilesDepotPaths();
+
+			EDialogFactoryResult DialogResult;
+			if (ShelvedFiles != LocalFiles)
+			{
+				const FText TextTitle = NSLOCTEXT("SourceControl.SubmitWindow", "ShelveConflictTitle", "Shelve - Local conflict");
+
+				const size_t MaxFilesToList = 5;
+				TArray<FString> ShelvedList;
+				TArray<FString> LocalList;
+				for (size_t i = 0; i < MaxFilesToList; ++i)
+				{
+					if (i < ShelvedFiles.Num())
+					{
+						ShelvedList.Add(ShelvedFiles[i]);
+					}
+
+					if (i < LocalFiles.Num())
+					{
+						LocalList.Add(LocalFiles[i]);
+					}
+				}
+				FString LocalListString = FString::Join(LocalList, TEXT("\n - "));
+				if (LocalFiles.Num() > MaxFilesToList)
+				{
+					LocalListString = FString::Printf(TEXT("%s\n - And %d other files"), *LocalListString, LocalFiles.Num() - MaxFilesToList);
+				}
+
+				FString ShelveListString = FString::Join(ShelvedList, TEXT("\n - "));
+				if (ShelvedFiles.Num() > MaxFilesToList)
+				{
+					ShelveListString = FString::Printf(TEXT("%s\n - And %d other files"), *ShelveListString, ShelvedFiles.Num() - MaxFilesToList);
+				}
+
+				const FString Description =
+					FString::Printf(TEXT("The shelve filelist does not match the local filelist, due to p4 restrictions submit tool can only submit local content do you want to continue with the submit?\nLocal Files:\n - %s\n\nShelved Files:\n - %s"),
+						*LocalListString,
+						*ShelveListString);
+				DialogResult = FDialogFactory::ShowDialog(TextTitle, FText::FromString(Description), TArray<FString>{TEXT("Delete Shelve and Submit"), TEXT("Cancel")});
+			}
+			else
+			{
+				const FText TextTitle = NSLOCTEXT("SourceControl.SubmitWindow", "DeleteShelvedFilesDialogTitle", "Delete shelved files?");
+				const FText TextDescription = NSLOCTEXT("SourceControl.SubmitWindow", "DeleteShelvedFilesDialogDescription", "There are shelved files in this changelist. Do you want to delete your shelf?\nIf you do not, the submit will be cancelled.");
+
+				DialogResult = FDialogFactory::ShowConfirmDialog(TextTitle, TextDescription);
+			}
+
+			if (DialogResult == EDialogFactoryResult::Confirm)
+			{
+				ChangeState(ESubmitToolAppState::Submitting, SubmitToolState == ESubmitToolAppState::SubmitLocked && bIsUserInAllowlist);
+				DeleteShelvedFiles();
+			}
 		}
 		else
 		{
+			ChangeState(ESubmitToolAppState::Submitting, SubmitToolState == ESubmitToolAppState::SubmitLocked && bIsUserInAllowlist);
 			RevertUnchangedAndSubmit();
 		}
 	}
@@ -551,7 +604,7 @@ void FModelInterface::RequestIntegration() const
 {
 	FNIntegrationService->RequestIntegration(FOnBooleanValueChanged::CreateLambda([this](bool bSuccess)
 		{
-			if(bSuccess)
+			if (bSuccess)
 			{
 				ChangeState(ESubmitToolAppState::Finished);
 			}
@@ -561,17 +614,17 @@ void FModelInterface::RefreshStateBasedOnFiles()
 {
 	TArray<FSourceControlStateRef> LocalFiles = ChangelistService->GetFilesInCL();
 
-	if(LocalFiles.IsEmpty())
+	if (LocalFiles.IsEmpty())
 	{
 		const TArray<FSourceControlStateRef>& ShelvedFiles = ChangelistService->GetShelvedFilesInCL();
-		if(ShelvedFiles.IsEmpty())
+		if (ShelvedFiles.IsEmpty())
 		{
 			UE_LOG(LogSubmitTool, Error, TEXT("There are no files in CL %s, SUBMIT IS DISABLED"), *ChangelistService->GetCLID());
 			ChangeState(ESubmitToolAppState::Errored);
 		}
 		else
 		{
-			if(!HasSubmitToolTag())
+			if (!HasSubmitToolTag())
 			{
 				UE_LOG(LogSubmitTool, Warning, TEXT("This CL hasn't been validated and there are no local files. You need to unshelve and run validations."), *ChangelistService->GetCLID());
 				ChangeState(ESubmitToolAppState::Errored);
@@ -579,7 +632,7 @@ void FModelInterface::RefreshStateBasedOnFiles()
 			else
 			{
 				UE::Tasks::Launch(UE_SOURCE_LOCATION, [this] {
-					if(P4LockdownService->ArePathsInLockdown(ChangelistService->GetShelvedFilesDepotPaths(), bIsUserInAllowlist))
+					if (P4LockdownService->ArePathsInLockdown(ChangelistService->GetShelvedFilesDepotPaths(), bIsUserInAllowlist))
 					{
 						UE_LOG(LogSubmitTool, Log, TEXT("There are no local files in CL %s, Submit is disabled but you can still request an Integration with your shelved files"), *ChangelistService->GetCLID());
 						ChangeState(ESubmitToolAppState::SubmitLocked);
@@ -596,24 +649,24 @@ void FModelInterface::RefreshStateBasedOnFiles()
 	else
 	{
 		const TArray<FSCCStream*>& Streams = SourceControlService->GetClientStreams();
-		if(!Streams.IsEmpty())
+		if (!Streams.IsEmpty())
 		{
 			FString StreamsMsg = FString::JoinBy(Streams, TEXT(" -> "), [](const FSCCStream* InStr) { return InStr->Name; });
-			for(const FString& File : ChangelistService->GetFilesDepotPaths())
+			for (const FString& File : ChangelistService->GetFilesDepotPaths())
 			{
 				bool bMappedToView = false;
 
-				for(const FSCCStream* Str : Streams)
+				for (const FSCCStream* Str : Streams)
 				{
-					if(File.StartsWith(Str->Name))
+					if (File.StartsWith(Str->Name))
 					{
 						bMappedToView = true;
-						break;						
+						break;
 					}
 
-					for(const FString& ImportStream : Str->AdditionalImportPaths)
+					for (const FString& ImportStream : Str->AdditionalImportPaths)
 					{
-						if(File.StartsWith(ImportStream))
+						if (File.StartsWith(ImportStream))
 						{
 							bMappedToView = true;
 							break;
@@ -621,7 +674,7 @@ void FModelInterface::RefreshStateBasedOnFiles()
 					}
 				}
 
-				if(!bMappedToView)
+				if (!bMappedToView)
 				{
 					UE_LOG(LogSubmitTool, Warning, TEXT("File %s is not in the stream that the workspace is set to: %s"), *File, *StreamsMsg);
 				}
@@ -629,7 +682,7 @@ void FModelInterface::RefreshStateBasedOnFiles()
 		}
 
 		UE::Tasks::Launch(UE_SOURCE_LOCATION, [this] {
-			if(P4LockdownService->ArePathsInLockdown(ChangelistService->GetFilesDepotPaths(), bIsUserInAllowlist))
+			if (P4LockdownService->ArePathsInLockdown(ChangelistService->GetFilesDepotPaths(), bIsUserInAllowlist))
 			{
 				ChangeState(ESubmitToolAppState::SubmitLocked);
 			}
@@ -643,9 +696,9 @@ void FModelInterface::RefreshStateBasedOnFiles()
 
 void FModelInterface::OnChangelistReady(bool bIsValid)
 {
-	if(SubmitToolState == ESubmitToolAppState::Initializing)
+	if (SubmitToolState == ESubmitToolAppState::Initializing)
 	{
-		if(bIsValid)
+		if (bIsValid)
 		{
 			UE_LOG(LogSubmitTool, Log, TEXT("Retrieved information for CL %s"), *ChangelistService->GetCLID());
 			PreflightService->FetchPreflightInfo(true);
@@ -654,15 +707,15 @@ void FModelInterface::OnChangelistReady(bool bIsValid)
 			SwarmService->FetchReview(OnGetReviewComplete::CreateRaw(this, &FModelInterface::OnGetUsersFromSwarmCompleted));
 
 			RefreshStateBasedOnFiles();
-			if(!ChangelistService->GetFilesInCL().IsEmpty())
+			if (!ChangelistService->GetFilesInCL().IsEmpty())
 			{
 				UpdateSubmitToolTag(false);
 				ValidationService->CheckForTagSkips();
 
 				ETaskArea ValidateArea = ~ETaskArea::Changelist;
-				for(const FTag* Tag : TagService->GetTagsArray())
+				for (const FTag* Tag : TagService->GetTagsArray())
 				{
-					if(Tag->GetValues().Num() != 0)
+					if (Tag->GetValues().Num() != 0)
 					{
 						ValidateArea = ETaskArea::Everything;
 						break;
@@ -686,12 +739,12 @@ void FModelInterface::OnChangelistReady(bool bIsValid)
 void FModelInterface::RevertUnchangedAndSubmit()
 {
 	bool bHasEditOrAdd = false;
-	for(const FSourceControlStateRef& file : ChangelistService->GetFilesInCL())
+	for (const FSourceControlStateRef& file : ChangelistService->GetFilesInCL())
 	{
 		bHasEditOrAdd |= file->IsAdded() || file->IsCheckedOut();
 	}
 
-	if(bHasEditOrAdd)
+	if (bHasEditOrAdd)
 	{
 		ChangelistService->RevertUnchangedFilesAsync(RevertUnchangedCallback);
 	}
@@ -703,25 +756,25 @@ void FModelInterface::RevertUnchangedAndSubmit()
 
 void FModelInterface::Submit()
 {
-	if(PresubmitOperationsService->AreTasksPendingQueue())
+	if (PresubmitOperationsService->AreTasksPendingQueue())
 	{
 		// This will call submit again when it's done
 		PresubmitOperationsService->CheckForTagSkips();
-		if(PresubmitOperationsService->QueueAll())
+		if (PresubmitOperationsService->QueueAll())
 		{
 			return;
 		}
 	}
 
 	auto AddendumAccumulator = [](const TArray<FString>& InAddendums, const FString& InDescription, FString& InOutAccumulated) {
-		for(const FString& Str : InAddendums)
+		for (const FString& Str : InAddendums)
 		{
-			if(!InDescription.Contains(Str, ESearchCase::IgnoreCase))
+			if (!InDescription.Contains(Str, ESearchCase::IgnoreCase))
 			{
 				InOutAccumulated += (TEXT("\n") + Str);
 			}
 		}
-	};
+		};
 
 	const FString& CLDescription = GetCLDescription();
 	FString Addendums;
@@ -733,17 +786,17 @@ void FModelInterface::Submit()
 
 void FModelInterface::OnDeleteShelveOperationComplete(const FSourceControlOperationRef& Operation, ECommandResult::Type Result)
 {
-	if(SubmitToolState == ESubmitToolAppState::Submitting)
+	if (SubmitToolState == ESubmitToolAppState::Submitting)
 	{
-		if(Result == ECommandResult::Succeeded)
+		if (Result == ECommandResult::Succeeded)
 		{
 			RevertUnchangedAndSubmit();
 		}
-		else if(Result == ECommandResult::Failed)
+		else if (Result == ECommandResult::Failed)
 		{
 			ChangeState(ESubmitToolAppState::WaitingUserInput);
 		}
-		else if(Result == ECommandResult::Cancelled)
+		else if (Result == ECommandResult::Cancelled)
 		{
 			ChangeState(ESubmitToolAppState::WaitingUserInput);
 		}
@@ -756,18 +809,18 @@ void FModelInterface::OnDeleteShelveOperationComplete(const FSourceControlOperat
 
 void FModelInterface::OnRevertUnchangedOperationComplete(const FSourceControlOperationRef& Operation, ECommandResult::Type Result)
 {
-	if(SubmitToolState == ESubmitToolAppState::Submitting)
+	if (SubmitToolState == ESubmitToolAppState::Submitting)
 	{
 		// Revert Unchanged returns as failed if there were no files to revert, check ErrorMessages to see actual failures
-		if(Result == ECommandResult::Cancelled)
+		if (Result == ECommandResult::Cancelled)
 		{
 			ChangeState(ESubmitToolAppState::WaitingUserInput);
 		}
-		else if(Result == ECommandResult::Succeeded || Operation->GetResultInfo().ErrorMessages.Num() == 0)
+		else if (Result == ECommandResult::Succeeded || Operation->GetResultInfo().ErrorMessages.Num() == 0)
 		{
 			Submit();
 		}
-		else if(Result == ECommandResult::Failed)
+		else if (Result == ECommandResult::Failed)
 		{
 			ChangeState(ESubmitToolAppState::WaitingUserInput);
 		}
@@ -781,7 +834,7 @@ void FModelInterface::OnRevertUnchangedOperationComplete(const FSourceControlOpe
 
 void FModelInterface::OnPresubmitOperationsComplete(bool bInSuccess)
 {
-	if(bInSuccess)
+	if (bInSuccess)
 	{
 		Submit();
 	}
@@ -795,24 +848,36 @@ void FModelInterface::OnPresubmitOperationsComplete(bool bInSuccess)
 
 void FModelInterface::OnSubmitOperationComplete(const FSourceControlOperationRef& Operation, ECommandResult::Type Result)
 {
-	if(Result == ECommandResult::Succeeded)
+	if (Result == ECommandResult::Succeeded)
 	{
 		// We've submitted, or tried to submit and failed so we only let the user close the app
 		ChangeState(ESubmitToolAppState::Finished);
-		FTag* JiraTag = TagService->GetTagOfType(TEXT("JiraIssue"));
-		if (JiraTag != nullptr && JiraTag->GetValues().Num() != 0)
+		if(FSubmitToolUserPrefs::Get()->bOpenJiraOnSubmit)
 		{
-			for (const FString& JiraValue : JiraTag->GetValues())
+			TSharedRef<FCheckIn> CheckIn = StaticCastSharedRef<FCheckIn>(Operation);
+			const FString CLString = CheckIn->GetSuccessMessage().ToString().Replace(TEXT("Submitted changelist "), TEXT(""), ESearchCase::IgnoreCase).TrimStartAndEnd();
+			if (CLString.IsNumeric())
+			{
+				FPlatformApplicationMisc::ClipboardCopy(*CLString);
+				UE_LOG(LogSubmitToolP4, Log, TEXT("Submitted CL copied to clipboard: %s"), *CLString);
+			}
+
+
+			FTag* JiraTag = TagService->GetTagOfType(TEXT("JiraIssue"));
+			if (JiraTag != nullptr && JiraTag->GetValues().Num() != 0)
 			{
-				if (!JiraValue.Equals(TEXT("none"), ESearchCase::IgnoreCase) && !JiraValue.Equals(TEXT("nojira"), ESearchCase::IgnoreCase))
+				for (const FString& JiraValue : JiraTag->GetValues())
 				{
-					FString Url = FString::Printf(TEXT("https://%s/browse/%s}"), *Parameters.JiraParameters.ServerAddress, *JiraValue);
-					FPlatformProcess::LaunchURL(*Url, nullptr, nullptr);
+					if (!JiraValue.Equals(TEXT("none"), ESearchCase::IgnoreCase) && !JiraValue.Equals(TEXT("nojira"), ESearchCase::IgnoreCase))
+					{
+						FString Url = FString::Printf(TEXT("https://%s/browse/%s}"), *Parameters.JiraParameters.ServerAddress, *JiraValue);
+						FPlatformProcess::LaunchURL(*Url, nullptr, nullptr);
+					}
 				}
 			}
 		}
 
-		if(FSubmitToolUserPrefs::Get()->bCloseOnSubmit)
+		if (FSubmitToolUserPrefs::Get()->bCloseOnSubmit)
 		{
 			MainTab.Pin()->RequestCloseTab();
 		}
@@ -825,24 +890,24 @@ void FModelInterface::OnSubmitOperationComplete(const FSourceControlOperationRef
 
 bool FModelInterface::Tick(float InDeltaTime)
 {
-	switch(SubmitToolState)
+	switch (SubmitToolState)
 	{
-	case ESubmitToolAppState::WaitingUserInput:
-		if(SwarmService->IsRequestRunning() || JiraService->IsBlockingRequestRunning())
-		{
-			ChangeState(ESubmitToolAppState::P4BlockingOperation);
-		}
-		break;
+		case ESubmitToolAppState::WaitingUserInput:
+			if (SwarmService->IsRequestRunning() || JiraService->IsBlockingRequestRunning())
+			{
+				ChangeState(ESubmitToolAppState::P4BlockingOperation);
+			}
+			break;
 
-	case ESubmitToolAppState::P4BlockingOperation:
-		if(!SwarmService->IsRequestRunning() && !ChangelistService->IsP4OperationRunning() && !JiraService->IsBlockingRequestRunning())
-		{
-			ChangeState(ESubmitToolAppState::WaitingUserInput);
-		}
-		break;
+		case ESubmitToolAppState::P4BlockingOperation:
+			if (!SwarmService->IsRequestRunning() && !ChangelistService->IsP4OperationRunning() && !JiraService->IsBlockingRequestRunning())
+			{
+				ChangeState(ESubmitToolAppState::WaitingUserInput);
+			}
+			break;
 
-	default:
-		break;
+		default:
+			break;
 	}
 
 	return true;
@@ -851,7 +916,7 @@ bool FModelInterface::Tick(float InDeltaTime)
 void FModelInterface::ChangeState(ESubmitToolAppState newState, bool bForce)
 {
 	const ESubmitToolAppState CurrentState = SubmitToolState;
-	if(bForce)
+	if (bForce)
 	{
 		UE_LOG(LogSubmitToolDebug, Log, TEXT("Transitioned state from '%s' to '%s'"), *StaticEnum<ESubmitToolAppState>()->GetNameStringByValue(static_cast<int64>(SubmitToolState)), *StaticEnum<ESubmitToolAppState>()->GetNameStringByValue(static_cast<int64>(newState)));
 		SubmitToolState = newState;
@@ -859,10 +924,10 @@ void FModelInterface::ChangeState(ESubmitToolAppState newState, bool bForce)
 	}
 	else
 	{
-		if(SubmitToolAppState::AllowedTransitions.Contains(SubmitToolState))
+		if (SubmitToolAppState::AllowedTransitions.Contains(SubmitToolState))
 		{
 			const TArray<ESubmitToolAppState>& allowedStates = SubmitToolAppState::AllowedTransitions[SubmitToolState];
-			if(allowedStates.Contains(newState))
+			if (allowedStates.Contains(newState))
 			{
 				UE_LOG(LogSubmitToolDebug, Log, TEXT("Transitioned state from '%s' to '%s'"), *StaticEnum<ESubmitToolAppState>()->GetNameStringByValue(static_cast<int64>(SubmitToolState)), *StaticEnum<ESubmitToolAppState>()->GetNameStringByValue(static_cast<int64>(newState)));
 				SubmitToolState = newState;
@@ -877,12 +942,12 @@ void FModelInterface::ChangeState(ESubmitToolAppState newState, bool bForce)
 		{
 			UE_LOG(LogSubmitToolDebug, Warning, TEXT("Transition not allowed from '%s' to '%s'"), *StaticEnum<ESubmitToolAppState>()->GetNameStringByValue(static_cast<int64>(SubmitToolState)), *StaticEnum<ESubmitToolAppState>()->GetNameStringByValue(static_cast<int64>(newState)));
 		}
-	}	
+	}
 }
 
 void FModelInterface::OnGetUsersFromSwarmCompleted(const TUniquePtr<FSwarmReview>& InReview, const FString& InErrorMessage)
 {
-	if(!InReview.IsValid())
+	if (!InReview.IsValid())
 	{
 		UE_LOG(LogSubmitTool, Log, TEXT("Could not retrieve swarm review for current changelist. %s"), *InErrorMessage);
 		return;
@@ -890,48 +955,48 @@ void FModelInterface::OnGetUsersFromSwarmCompleted(const TUniquePtr<FSwarmReview
 
 	TArray<const FTag*> TargetTags;
 
-	for(const FTag* Tag : TagService->GetTagsArray())
+	for (const FTag* Tag : TagService->GetTagsArray())
 	{
-		if(Tag->Definition.InputSubType.Equals(TEXT("Swarm"), ESearchCase::IgnoreCase))
+		if (Tag->Definition.InputSubType.Equals(TEXT("Swarm"), ESearchCase::IgnoreCase))
 		{
 			TargetTags.Add(Tag);
 		}
 	}
 
-	if(!TargetTags.IsEmpty())
+	if (!TargetTags.IsEmpty())
 	{
 		TArray<FString> SwarmUserValues;
 
-		for(const TPair<FString, FSwarmReviewParticipant>& Participant : InReview->Participants)
+		for (const TPair<FString, FSwarmReviewParticipant>& Participant : InReview->Participants)
 		{
-			if(Participant.Key.Equals(InReview->Author, ESearchCase::IgnoreCase))
+			if (Participant.Key.Equals(InReview->Author, ESearchCase::IgnoreCase))
 			{
 				continue;
 			}
 
-			if(Participant.Value.Vote.Value == 1)
+			if (Participant.Value.Vote.Value == 1)
 			{
-				if(!SwarmUserValues.Contains(Participant.Key) && !SwarmUserValues.Contains(TEXT("@") + Participant.Key))
+				if (!SwarmUserValues.Contains(Participant.Key) && !SwarmUserValues.Contains(TEXT("@") + Participant.Key))
 				{
 					SwarmUserValues.Add(Participant.Key);
 				}
 			}
 		}
 
-		if(!SwarmUserValues.IsEmpty())
+		if (!SwarmUserValues.IsEmpty())
 		{
 			bool bApplied = false;
 
-			for(const FTag* Tag : TargetTags)
+			for (const FTag* Tag : TargetTags)
 			{
-				if(Tag->GetValues() != SwarmUserValues)
+				if (Tag->GetValues() != SwarmUserValues)
 				{
 					SetTagValues(*Tag, SwarmUserValues);
 					bApplied = true;
 				}
 			}
 
-			if(bApplied)
+			if (bApplied)
 			{
 				UE_LOG(LogSubmitTool, Log, TEXT("RB tag set to users that upvoted review '%d' Users: %s"), InReview->Id, *FString::Join(SwarmUserValues, TEXT(", ")));
 				UE_LOG(LogSubmitToolDebug, Log, TEXT("Re-running Tag validator after applying the #rb from swarm"));
@@ -943,7 +1008,7 @@ void FModelInterface::OnGetUsersFromSwarmCompleted(const TUniquePtr<FSwarmReview
 
 void FModelInterface::OnSwarmCreateCompleted(bool InResult, const FString& InErrorMessage)
 {
-	if(InResult)
+	if (InResult)
 	{
 		OnGetUsersFromSwarmCompleted(SwarmService->GetReview(), InErrorMessage);
 		ShowSwarmReview();
diff --git a/Engine/Source/Programs/SubmitTool/Private/Models/SubmitToolUserPrefs.h b/Engine/Source/Programs/SubmitTool/Private/Models/SubmitToolUserPrefs.h
--- a/Engine/Source/Programs/SubmitTool/Private/Models/SubmitToolUserPrefs.h
+++ b/Engine/Source/Programs/SubmitTool/Private/Models/SubmitToolUserPrefs.h
@@ -17,6 +17,9 @@ struct FSubmitToolUserPrefs
 	
 	UPROPERTY()
 	bool bCloseOnSubmit = true;
+
+	UPROPERTY()
+	bool bOpenJiraOnSubmit = true;
 	
 	UPROPERTY()
 	bool bExpandFilesInCL = false;
diff --git a/Engine/Source/Programs/SubmitTool/Private/View/Widgets/SubmitToolWidget.cpp b/Engine/Source/Programs/SubmitTool/Private/View/Widgets/SubmitToolWidget.cpp
--- a/Engine/Source/Programs/SubmitTool/Private/View/Widgets/SubmitToolWidget.cpp
+++ b/Engine/Source/Programs/SubmitTool/Private/View/Widgets/SubmitToolWidget.cpp
@@ -158,8 +158,7 @@ void SubmitToolWidget::Construct(const FArguments& InArgs)
 		.OnTextCommitted_Lambda([this](const FText& Text, ETextCommit::Type CommitType) { if(CommitType != ETextCommit::OnEnter) { ModelInterface->ValidateCLDescription(); }})
 		.AutoWrapText(true)
 		.OnIsTypedCharValid(FOnIsTypedCharValid::CreateLambda([](const TCHAR) { return true; }))
-		.IsReadOnly(false)
-		.IsEnabled_Static(&FModelInterface::GetInputEnabled);
+		.IsReadOnly_Lambda([]{ return !FModelInterface::GetInputEnabled();});
 
 	FString PerforceClientName;
 	FCmdLineParameters::Get().GetValue(FSubmitToolCmdLine::P4Client, PerforceClientName);
@@ -426,6 +425,38 @@ TSharedRef<SHorizontalBox> SubmitToolWidget::BuildButtonRow()
 		.HAlign(HAlign_Right)
 		.MaxWidth(700)
 		.AutoWidth()
+		[
+			SNew(SHorizontalBox)
+			+SHorizontalBox::Slot()
+			.AutoWidth()
+			.VAlign(VAlign_Center)
+			[
+				SNew(SCheckBox)
+				.IsChecked_Lambda([this]() { return FSubmitToolUserPrefs::Get()->bOpenJiraOnSubmit ? ECheckBoxState::Checked : ECheckBoxState::Unchecked; })
+				.OnCheckStateChanged_Lambda([this](ECheckBoxState InNewState) { FSubmitToolUserPrefs::Get()->bOpenJiraOnSubmit = !FSubmitToolUserPrefs::Get()->bOpenJiraOnSubmit; })
+			]
+			+SHorizontalBox::Slot()
+			.AutoWidth()
+			.VAlign(VAlign_Center)
+			[
+				SNew(SButton)
+				.ButtonStyle(FAppStyle::Get(), "InvisibleButton")
+				.IsFocusable(false)	
+				.OnClicked_Lambda([this]() { FSubmitToolUserPrefs::Get()->bOpenJiraOnSubmit = !FSubmitToolUserPrefs::Get()->bOpenJiraOnSubmit; return FReply::Handled(); })
+				[
+					SNew(STextBlock)
+					.Justification(ETextJustify::Left)
+					.Font(FAppStyle::Get().GetFontStyle("SmallFont"))
+					.MinDesiredWidth(60)
+					.Text(FText::FromString(TEXT("Open Ticket on Submit")))
+				]
+			]
+		]
+		+SHorizontalBox::Slot()
+		.FillWidth(1.f)
+		.HAlign(HAlign_Right)
+		.MaxWidth(700)
+		.AutoWidth()
 		[
 			SNew(SHorizontalBox)
 			+SHorizontalBox::Slot()
@@ -752,18 +783,7 @@ FReply SubmitToolWidget::SubmitClicked()
 
 		if(ModelInterface->IsCLValid())
 		{
-			if(ModelInterface->HasShelvedFiles())
-			{
-				EDialogFactoryResult DeleteShelveResult = ShowDeleteShelveDialog();
-				if(DeleteShelveResult == EDialogFactoryResult::Confirm)
-				{
-					ModelInterface->StartSubmitProcess();
-				}
-			}
-			else
-			{
-				ModelInterface->StartSubmitProcess();
-			}
+			ModelInterface->StartSubmitProcess();
 		}
 	}
 
@@ -791,16 +811,6 @@ FReply SubmitToolWidget::ValidateClicked()
 	return FReply::Handled();
 }
 
-EDialogFactoryResult SubmitToolWidget::ShowDeleteShelveDialog() const
-{
-	const FText TextTitle = NSLOCTEXT("SourceControl.SubmitWindow", "DeleteShelvedFilesDialogTitle", "Delete shelved files?");
-	const FText TextDescription = NSLOCTEXT("SourceControl.SubmitWindow", "DeleteShelvedFilesDialogDescription", "There are shelved files in this changelist. Do you want to delete your shelf?\nIf you do not, the submit will be cancelled.");
-	const FText TextCancel = NSLOCTEXT("SourceControl.SubmitWindow", "DeleteShelvedFilesDialogCancel", "Cancel submit");
-	const FText TextConfirm = NSLOCTEXT("SourceControl.SubmitWindow", "DeleteShelvedFilesDialogConfirm", "Delete shelve");
-
-	return FDialogFactory::ShowConfirmDialog(TextTitle, TextDescription);
-}
-
 FText SubmitToolWidget::GetMainButtonText() const
 {
 	FText FinalText;
diff --git a/Engine/Source/Programs/SubmitTool/Private/View/Widgets/TagWidget.cpp b/Engine/Source/Programs/SubmitTool/Private/View/Widgets/TagWidget.cpp
--- a/Engine/Source/Programs/SubmitTool/Private/View/Widgets/TagWidget.cpp
+++ b/Engine/Source/Programs/SubmitTool/Private/View/Widgets/TagWidget.cpp
@@ -164,7 +164,7 @@ void STagWidget::Construct(const FArguments& InArgs)
 			}})
 			.Text_Lambda([&tag = Tag] {return FText::FromString(tag->GetValuesText()); })
 			.OnTextChanged(this, &STagWidget::OnTextChanged)
-			.IsEnabled_Lambda([]()->bool { return FModelInterface::GetInputEnabled(); })
+			.IsReadOnly_Lambda([] { return !FModelInterface::GetInputEnabled(); })
 		];
 	}
 
