Fix Dangling OIDC token process getting stuck for submit tool

diff --git a/Engine/Source/Programs/SubmitTool/Private/Logic/CredentialsService.cpp b/Engine/Source/Programs/SubmitTool/Private/Logic/CredentialsService.cpp
--- a/Engine/Source/Programs/SubmitTool/Private/Logic/CredentialsService.cpp
+++ b/Engine/Source/Programs/SubmitTool/Private/Logic/CredentialsService.cpp
@@ -29,6 +29,16 @@ FCredentialsService::FCredentialsService(const FOAuthTokenParams& InOAuthParamet
 	LoadCredentials();
 }
 
+FCredentialsService::~FCredentialsService()
+{
+	if(OIDCProcess.IsValid())
+	{
+		OIDCProcess->Stop();
+	}
+
+	GetOIDCTask.Wait();
+}
+
 void FCredentialsService::LoadKey()
 {
 	if(!IFileManager::Get().FileExists(*GetKeyFilepath()))
@@ -159,15 +169,20 @@ void FCredentialsService::GetOIDCToken()
 					}
 					});
 
-				FProcessWrapper GetToken = FProcessWrapper(TEXT("Oidc"), Parameters.OAuthTokenTool, FString::Printf(TEXT("%s"), *Parameters.OAuthArgs), FOnCompleted(), OutputLineProcess);
-				GetToken.Start(true);
+				OIDCProcess = MakeUnique<FProcessWrapper>(TEXT("Oidc"), Parameters.OAuthTokenTool, FString::Printf(TEXT("%s"), *Parameters.OAuthArgs), FOnCompleted(), OutputLineProcess);
+				OIDCProcess->Start(true);
 
-				if(GetToken.ExitCode != 0)
+				if(!OIDCProcess.IsValid() || OIDCProcess->ExitCode != 0)
 				{
 					UE_LOG(LogSubmitTool, Warning, TEXT("Couldn't obtain OIDC credentials"));
+					if(OIDCProcess.IsValid())
+					{
+						OIDCProcess = nullptr;
+					}
 					return false;
 				}
 
+				OIDCProcess = nullptr;
 				return true;
 			});
 	}
diff --git a/Engine/Source/Programs/SubmitTool/Private/Logic/CredentialsService.h b/Engine/Source/Programs/SubmitTool/Private/Logic/CredentialsService.h
--- a/Engine/Source/Programs/SubmitTool/Private/Logic/CredentialsService.h
+++ b/Engine/Source/Programs/SubmitTool/Private/Logic/CredentialsService.h
@@ -6,11 +6,13 @@
 #include "Tasks/Task.h"
 #include "Services/Interfaces/ISubmitToolService.h"
 
+class FProcessWrapper;
+
 class FCredentialsService : public ISubmitToolService
 {
 public:
 	FCredentialsService(const FOAuthTokenParams& InOAuthParameters);
-	~FCredentialsService() = default;
+	virtual ~FCredentialsService();
 
 	bool HasCredentials() const
 	{
@@ -74,6 +76,7 @@ private:
 	static const FString GetKeyFilepath();
 
 	UE::Tasks::TTask<bool> GetOIDCTask;
+	TUniquePtr<FProcessWrapper> OIDCProcess;
 	FString OIDCToken;
 	FDateTime TokenExpiration;
 	const FOAuthTokenParams Parameters;
