Add the current changelist when raising an analytics event for a validator.

#rnx

diff --git a/Engine/Source/Programs/SubmitTool/Private/Logic/Validators/ValidatorBase.cpp b/Engine/Source/Programs/SubmitTool/Private/Logic/Validators/ValidatorBase.cpp
--- a/Engine/Source/Programs/SubmitTool/Private/Logic/Validators/ValidatorBase.cpp
+++ b/Engine/Source/Programs/SubmitTool/Private/Logic/Validators/ValidatorBase.cpp
@@ -210,12 +210,16 @@ const FString FValidatorBase::GetStatusText() const
 
 const TArray<FAnalyticsEventAttribute> FValidatorBase::GetTelemetryAttributes() const
 {
+	const TSharedPtr<FChangelistService> ChangelistService = ServiceProvider.Pin()->GetService<FChangelistService>();
+	FString ChangelistNumber = ChangelistService->GetCLID();
+
 	return MakeAnalyticsEventAttributeArray(
 		TEXT("ValidatorID"), *GetValidatorNameId().ToString(),
 		TEXT("ValidatorName"), *GetValidatorName(),
 		TEXT("Status"), GetHasPassed(),
 		TEXT("Runtime"), RunTime,
-		TEXT("Stream"), ServiceProvider.Pin()->GetService<ISTSourceControlService>()->GetCurrentStreamName()
+		TEXT("Stream"), ServiceProvider.Pin()->GetService<ISTSourceControlService>()->GetCurrentStreamName(),
+		TEXT("PendingChangelist"), MoveTemp(ChangelistNumber)
 		);
 }
 
