Fix for output log filter having duplicated and incorrect categories
Added SourceControl as enabled category

diff --git a/Engine/Source/Programs/SubmitTool/Private/View/Widgets/SubmitToolWidget.cpp b/Engine/Source/Programs/SubmitTool/Private/View/Widgets/SubmitToolWidget.cpp
--- a/Engine/Source/Programs/SubmitTool/Private/View/Widgets/SubmitToolWidget.cpp
+++ b/Engine/Source/Programs/SubmitTool/Private/View/Widgets/SubmitToolWidget.cpp
@@ -622,16 +622,18 @@ TSharedRef<SWidget> SubmitToolWidget::BuildOutputLogWidget()
 	// ----------------------------------------------------------------------
 	FOutputLogCreationParams SummaryLogCreationParams;
 
-	// Initially only allow Submit Tool Log as the log filter
 	SummaryLogCreationParams.AllowAsInitialLogCategory = FAllowLogCategoryCallback::CreateLambda([](const FName LogCategory) {
-		return LogCategory == LogValidators.GetCategoryName();
+		return LogCategory == LogSubmitTool.GetCategoryName();
 		});
 
 	SummaryLogCreationParams.DefaultCategorySelection.Emplace(LogSubmitTool.GetCategoryName(), true);
+	SummaryLogCreationParams.DefaultCategorySelection.Emplace(LogValidators.GetCategoryName(), false);
 	SummaryLogCreationParams.DefaultCategorySelection.Emplace(LogValidatorsResult.GetCategoryName(), true);
 	SummaryLogCreationParams.DefaultCategorySelection.Emplace(LogPresubmitResult.GetCategoryName(), true);
 	SummaryLogCreationParams.DefaultCategorySelection.Emplace(LogOutputDevice.GetCategoryName(), true);
 	SummaryLogCreationParams.DefaultCategorySelection.Emplace(LogSubmitToolP4.GetCategoryName(), true);
+	SummaryLogCreationParams.DefaultCategorySelection.Emplace(LogPresubmit.GetCategoryName(), false);
+	SummaryLogCreationParams.DefaultCategorySelection.Emplace("SourceControl", true);
 
 	SummaryLogDockTab = SNew(SDockTab)
 		.CanEverClose(false)
@@ -642,7 +644,7 @@ TSharedRef<SWidget> SubmitToolWidget::BuildOutputLogWidget()
 			OutputLogModule.MakeOutputLogWidget(SummaryLogCreationParams)
 		];
 	
-	LogTabManager->RegisterTabSpawner(SummaryTabId, FOnSpawnTab::CreateLambda([this, &OutputLogModule, &SummaryLogCreationParams](const FSpawnTabArgs& SpawnArgs) {	return SummaryLogDockTab.ToSharedRef(); }));
+	LogTabManager->RegisterTabSpawner(SummaryTabId, FOnSpawnTab::CreateLambda([this](const FSpawnTabArgs& SpawnArgs) {	return SummaryLogDockTab.ToSharedRef(); }));
 
 	// ----------------------------------------------------------------------
 	// Validators Log Tab
@@ -652,7 +654,14 @@ TSharedRef<SWidget> SubmitToolWidget::BuildOutputLogWidget()
 		return LogCategory == LogValidators.GetCategoryName();
 		});
 
+	ValidatorLogCreationParams.DefaultCategorySelection.Emplace(LogSubmitTool.GetCategoryName(), false);
 	ValidatorLogCreationParams.DefaultCategorySelection.Emplace(LogValidators.GetCategoryName(), true);
+	ValidatorLogCreationParams.DefaultCategorySelection.Emplace(LogValidatorsResult.GetCategoryName(), false);
+	ValidatorLogCreationParams.DefaultCategorySelection.Emplace(LogPresubmitResult.GetCategoryName(), false);
+	ValidatorLogCreationParams.DefaultCategorySelection.Emplace(LogOutputDevice.GetCategoryName(), false);
+	ValidatorLogCreationParams.DefaultCategorySelection.Emplace(LogSubmitToolP4.GetCategoryName(), false);
+	ValidatorLogCreationParams.DefaultCategorySelection.Emplace(LogPresubmit.GetCategoryName(), false);
+
 
 	ValidatorLogTab = SNew(SDockTab)
 		.CanEverClose(false)
@@ -663,7 +672,7 @@ TSharedRef<SWidget> SubmitToolWidget::BuildOutputLogWidget()
 			OutputLogModule.MakeOutputLogWidget(ValidatorLogCreationParams)
 		];
 
-	LogTabManager->RegisterTabSpawner(ValidatorLogTabId, FOnSpawnTab::CreateLambda([&OutputLogModule, this](const FSpawnTabArgs& SpawnArgs) { return ValidatorLogTab.ToSharedRef(); }));
+	LogTabManager->RegisterTabSpawner(ValidatorLogTabId, FOnSpawnTab::CreateLambda([this](const FSpawnTabArgs& SpawnArgs) { return ValidatorLogTab.ToSharedRef(); }));
 
 	// ----------------------------------------------------------------------
 	// Pre Submit Log Tab
@@ -674,6 +683,12 @@ TSharedRef<SWidget> SubmitToolWidget::BuildOutputLogWidget()
 		});
 
 	PresubmitLogCreationParams.DefaultCategorySelection.Emplace(LogPresubmit.GetCategoryName(), true);
+	PresubmitLogCreationParams.DefaultCategorySelection.Emplace(LogSubmitTool.GetCategoryName(), false);
+	PresubmitLogCreationParams.DefaultCategorySelection.Emplace(LogValidators.GetCategoryName(), false);
+	PresubmitLogCreationParams.DefaultCategorySelection.Emplace(LogValidatorsResult.GetCategoryName(), false);
+	PresubmitLogCreationParams.DefaultCategorySelection.Emplace(LogPresubmitResult.GetCategoryName(), false);
+	PresubmitLogCreationParams.DefaultCategorySelection.Emplace(LogOutputDevice.GetCategoryName(), false);
+	PresubmitLogCreationParams.DefaultCategorySelection.Emplace(LogSubmitToolP4.GetCategoryName(), false);
 
 	PresubmitLogTab = SNew(SDockTab)
 		.CanEverClose(false)
@@ -684,7 +699,7 @@ TSharedRef<SWidget> SubmitToolWidget::BuildOutputLogWidget()
 			OutputLogModule.MakeOutputLogWidget(PresubmitLogCreationParams)
 		];
 
-	LogTabManager->RegisterTabSpawner(PresubmitLogTabId, FOnSpawnTab::CreateLambda([&OutputLogModule, this](const FSpawnTabArgs& SpawnArgs) { return PresubmitLogTab.ToSharedRef(); }));
+	LogTabManager->RegisterTabSpawner(PresubmitLogTabId, FOnSpawnTab::CreateLambda([this](const FSpawnTabArgs& SpawnArgs) { return PresubmitLogTab.ToSharedRef(); }));
 
 
 	// ----------------------------------------------------------------------
