Fixing task logic evaluation being delayed one frame and producing incorrect telemetry for synchronous validators

diff --git a/Engine/Source/Programs/SubmitTool/Private/Logic/TasksService.cpp b/Engine/Source/Programs/SubmitTool/Private/Logic/TasksService.cpp
--- a/Engine/Source/Programs/SubmitTool/Private/Logic/TasksService.cpp
+++ b/Engine/Source/Programs/SubmitTool/Private/Logic/TasksService.cpp
@@ -325,22 +325,6 @@ bool FTasksService::Tick(float InDeltaTime)
 	{
 		const TSharedPtr<FValidatorBase>& Task = TaskPair.Value;
 
-		if (Task->Definition->IsRequired)
-		{
-			bCurrentCLValidState = bCurrentCLValidState && Task->GetHasPassed();
-		}
-		else if(Task->Definition->bRequireCompleteWhenOptional)
-		{
-			bCurrentCLValidState = bCurrentCLValidState && !Task->GetIsRunningOrQueued();
-		}
-
-		bAnyTaskRunningOrQueued = bAnyTaskRunningOrQueued || Task->GetIsRunningOrQueued();
-
-		if(!Task->GetHasPassed())
-		{
-			InvalidateDependants(Task);
-		}
-
 		if(Task->GetIsQueued())
 		{
 			bool bDependenciesSatisfied = true;
@@ -392,6 +376,22 @@ bool FTasksService::Tick(float InDeltaTime)
 		{
 			Task->Tick(InDeltaTime);
 		}
+
+		if (Task->Definition->IsRequired)
+		{
+			bCurrentCLValidState = bCurrentCLValidState && Task->GetHasPassed();
+		}
+		else if (Task->Definition->bRequireCompleteWhenOptional)
+		{
+			bCurrentCLValidState = bCurrentCLValidState && !Task->GetIsRunningOrQueued();
+		}
+
+		bAnyTaskRunningOrQueued = bAnyTaskRunningOrQueued || Task->GetIsRunningOrQueued();
+
+		if (!Task->GetIsRunningOrQueued() && !Task->GetHasPassed())
+		{
+			InvalidateDependants(Task);
+		}
 	}
 
 	if(bLastTasksRunState != bCurrentCLValidState)
