Add support for using $engine in the submittool and it's configuration file.

#rnx

- At the moment this is the functional equivalent of using ($root)/Engine but it is preferable to have a short hand for the engine location that is not so closely bound to pre-existing directory structure.

diff --git a/Engine/Source/Programs/SubmitTool/Private/Configuration/Configuration.cpp b/Engine/Source/Programs/SubmitTool/Private/Configuration/Configuration.cpp
--- a/Engine/Source/Programs/SubmitTool/Private/Configuration/Configuration.cpp
+++ b/Engine/Source/Programs/SubmitTool/Private/Configuration/Configuration.cpp
@@ -25,6 +25,8 @@ void FConfiguration::Init()
 
 	Configuration.Values.Add(TEXT("$(root)"), RootDir);
 
+	Configuration.Values.Add(TEXT("$(engine)"), RootDir + TEXT("/Engine"));
+
 	FString Changelist;
 	FCmdLineParameters::Get().GetValue(FSubmitToolCmdLine::P4ChangeList, Changelist);
 	Configuration.Values.Add(TEXT("$(CL)"), Changelist);
