Add support for using $(UBTPlatform) in the SubmitTool

#rnx

- Use of $(UBTPlatform) will be replaced with the name of the platform that UBT used to compile the current SubmitTool with.
- This is useful to use when you want to invoke UBT to build a project or program with the current platform that the SubmitTool is running with.

diff --git a/Engine/Source/Programs/SubmitTool/Private/Configuration/Configuration.cpp b/Engine/Source/Programs/SubmitTool/Private/Configuration/Configuration.cpp
--- a/Engine/Source/Programs/SubmitTool/Private/Configuration/Configuration.cpp
+++ b/Engine/Source/Programs/SubmitTool/Private/Configuration/Configuration.cpp
@@ -26,7 +26,8 @@ void FConfiguration::Init()
 	Configuration.Values.Add(TEXT("$(root)"), RootDir);
 
 	Configuration.Values.Add(TEXT("$(engine)"), RootDir + TEXT("/Engine"));
-
+	Configuration.Values.Add(TEXT("$(UBTPlatform)"), FGenericPlatformMisc::GetUBTPlatform());
+	
 	FString Changelist;
 	FCmdLineParameters::Get().GetValue(FSubmitToolCmdLine::P4ChangeList, Changelist);
 	Configuration.Values.Add(TEXT("$(CL)"), Changelist);
