Adding two new aliases $(RunUAT) and $(RunUBT) to SubmitTool

#rnx

- These aliases resolve to "Engine/Build/BatchFiles/RunUAT.bat" and "Engine/Build/BatchFiles/RunUBT.bat" on windows and resolves to their *.sh variants for other platforms.
- This ensures that people writing new validators don't need to worry about the correct extension for each platform that the submit tool might eb run from.

diff --git a/Engine/Source/Programs/SubmitTool/Private/Configuration/Configuration.cpp b/Engine/Source/Programs/SubmitTool/Private/Configuration/Configuration.cpp
--- a/Engine/Source/Programs/SubmitTool/Private/Configuration/Configuration.cpp
+++ b/Engine/Source/Programs/SubmitTool/Private/Configuration/Configuration.cpp
@@ -25,7 +25,21 @@ void FConfiguration::Init()
 
 	Configuration.Values.Add(TEXT("$(root)"), RootDir);
 
-	Configuration.Values.Add(TEXT("$(engine)"), RootDir + TEXT("/Engine"));
+	const FString EngineDir = RootDir + TEXT("/Engine");
+	Configuration.Values.Add(TEXT("$(engine)"), EngineDir);
+
+	const FString BatchFileDir = EngineDir + TEXT("/Build/BatchFiles");
+
+#if PLATFORM_WINDOWS
+	Configuration.Values.Add(TEXT("$(RunUAT)"), BatchFileDir + TEXT("/RunUAT.bat"));
+	Configuration.Values.Add(TEXT("$(RunUBT)"), BatchFileDir + TEXT("/RunUBT.bat"));
+#elif (PLATFORM_LINUX || PLATFORM_MAC)
+	Configuration.Values.Add(TEXT("$(RunUAT)"), BatchFileDir + TEXT("/RunUAT.sh"));
+	Configuration.Values.Add(TEXT("$(RunUBT)"), BatchFileDir + TEXT("/RunUBT.sh"));
+#else
+	UE_LOG(LogSubmitTool, Error, TEXT("Unknown platform, cannot resolve aliases $(RunUAT) and $(RunUBT)"));
+#endif
+
 	Configuration.Values.Add(TEXT("$(UBTPlatform)"), FGenericPlatformMisc::GetUBTPlatform());
 	
 	FString Changelist;
