Separated the concept of swarm and swarm approved to allow for different tags driving reviewers and approvers.

Added max height and scrolling to the tags section

diff --git a/Engine/Source/Programs/SubmitTool/Private/Models/ModelInterface.cpp b/Engine/Source/Programs/SubmitTool/Private/Models/ModelInterface.cpp
--- a/Engine/Source/Programs/SubmitTool/Private/Models/ModelInterface.cpp
+++ b/Engine/Source/Programs/SubmitTool/Private/Models/ModelInterface.cpp
@@ -957,7 +957,7 @@ void FModelInterface::OnGetUsersFromSwarmCompleted(const TUniquePtr<FSwarmReview
 
 	for (const FTag* Tag : TagService->GetTagsArray())
 	{
-		if (Tag->Definition.InputSubType.Equals(TEXT("Swarm"), ESearchCase::IgnoreCase))
+		if (Tag->Definition.InputSubType.Equals(TEXT("SwarmApproved"), ESearchCase::IgnoreCase))
 		{
 			TargetTags.Add(Tag);
 		}
diff --git a/Engine/Source/Programs/SubmitTool/Private/View/Widgets/SubmitToolWidget.cpp b/Engine/Source/Programs/SubmitTool/Private/View/Widgets/SubmitToolWidget.cpp
--- a/Engine/Source/Programs/SubmitTool/Private/View/Widgets/SubmitToolWidget.cpp
+++ b/Engine/Source/Programs/SubmitTool/Private/View/Widgets/SubmitToolWidget.cpp
@@ -96,6 +96,8 @@ void SubmitToolWidget::Construct(const FArguments& InArgs)
 	/**** Tags ****/
 	Contents->AddSlot()
 		.AutoHeight()
+		.MaxHeight(150)
+		.Padding(0,0,0,3)
 		[
 			SNew(STagSectionWidget)
 			.ParentWindow(InArgs._ParentWindow)
diff --git a/Engine/Source/Programs/SubmitTool/Private/View/Widgets/TagSectionWidget.cpp b/Engine/Source/Programs/SubmitTool/Private/View/Widgets/TagSectionWidget.cpp
--- a/Engine/Source/Programs/SubmitTool/Private/View/Widgets/TagSectionWidget.cpp
+++ b/Engine/Source/Programs/SubmitTool/Private/View/Widgets/TagSectionWidget.cpp
@@ -6,6 +6,7 @@
 #include "Models/ModelInterface.h"
 #include "Models/Tag.h"
 #include "View/Widgets/SJiraWidget.h"
+#include "Widgets/Layout/SScrollBox.h"
 
 
 #define LOCTEXT_NAMESPACE "TagSectionWidget"
@@ -21,14 +22,20 @@ void STagSectionWidget::Construct(const FArguments& InArgs)
 
 	ChildSlot
 	[
-		SNew(SHorizontalBox)
-		+ SHorizontalBox::Slot()
+		SNew(SScrollBox)
+		.Orientation(EOrientation::Orient_Vertical)
+		+SScrollBox::Slot()
+		.AutoSize()
 		[
-			SAssignNew(ColumnOne, SVerticalBox)
-		]
-		+ SHorizontalBox::Slot()
-		[
-			SAssignNew(ColumnTwo, SVerticalBox)
+			SNew(SHorizontalBox)
+			+SHorizontalBox::Slot()
+			[
+				SAssignNew(ColumnOne, SVerticalBox)
+			]
+			+SHorizontalBox::Slot()
+			[
+				SAssignNew(ColumnTwo, SVerticalBox)
+			]
 		]
 	];
 
diff --git a/Engine/Source/Programs/SubmitTool/Private/View/Widgets/TagWidget.cpp b/Engine/Source/Programs/SubmitTool/Private/View/Widgets/TagWidget.cpp
--- a/Engine/Source/Programs/SubmitTool/Private/View/Widgets/TagWidget.cpp
+++ b/Engine/Source/Programs/SubmitTool/Private/View/Widgets/TagWidget.cpp
@@ -170,8 +170,8 @@ void STagWidget::Construct(const FArguments& InArgs)
 
 	// optionally add the perforce users input
 	if (this->Tag->Definition.InputType.Equals("PerforceUser", ESearchCase::IgnoreCase))
-	{	
-		if (this->Tag->Definition.InputSubType.Equals("Swarm", ESearchCase::IgnoreCase))
+	{
+		if (this->Tag->Definition.InputSubType.Equals("SwarmApproved", ESearchCase::IgnoreCase))
 		{
 			HorizontalBox->AddSlot()
 				.Padding(3, 0, 0, 0)
@@ -203,7 +203,9 @@ void STagWidget::Construct(const FArguments& InArgs)
 								]
 						]
 				];
-
+		}
+		else if (this->Tag->Definition.InputSubType.Equals("Swarm", ESearchCase::IgnoreCase))
+		{
 			HorizontalBox->AddSlot()
 				.Padding(3, 0, 0, 0)
 				.AutoWidth()
