Fix preflight validator running indefinitely when horde connection failed in an unrecoverable way

diff --git a/Engine/Source/Programs/SubmitTool/Private/Logic/PreflightService.cpp b/Engine/Source/Programs/SubmitTool/Private/Logic/PreflightService.cpp
--- a/Engine/Source/Programs/SubmitTool/Private/Logic/PreflightService.cpp
+++ b/Engine/Source/Programs/SubmitTool/Private/Logic/PreflightService.cpp
@@ -263,6 +263,8 @@ void FPreflightService::FetchPreflightInfo(bool bRequeue, const FString& InOAuth
 					{
 						UE_LOG(LogSubmitTool, Warning, TEXT("Unable to connect to horde. Connection error, no response."));
 					}
+
+					OnHordeConnectionFailed.Broadcast();
 					return;
 				}
 
diff --git a/Engine/Source/Programs/SubmitTool/Private/Logic/PreflightService.h b/Engine/Source/Programs/SubmitTool/Private/Logic/PreflightService.h
--- a/Engine/Source/Programs/SubmitTool/Private/Logic/PreflightService.h
+++ b/Engine/Source/Programs/SubmitTool/Private/Logic/PreflightService.h
@@ -68,6 +68,7 @@ public:
 	bool SelectPreflightTemplate(FPreflightTemplateDefinition& OutTemplate) const;
 
 	FOnPreflightDataUpdated OnPreflightDataUpdated;
+	FSimpleMulticastDelegate OnHordeConnectionFailed;
 
 private:
 	void QueueFetch(bool bRequeue, float InSeconds);
diff --git a/Engine/Source/Programs/SubmitTool/Private/Logic/Validators/PreflightValidator.cpp b/Engine/Source/Programs/SubmitTool/Private/Logic/Validators/PreflightValidator.cpp
--- a/Engine/Source/Programs/SubmitTool/Private/Logic/Validators/PreflightValidator.cpp
+++ b/Engine/Source/Programs/SubmitTool/Private/Logic/Validators/PreflightValidator.cpp
@@ -44,6 +44,12 @@ bool FPreflightValidator::Validate(const FString& InCLDescription, const TArray<
 		PreflightUpdateHandler = PreflightService->OnPreflightDataUpdated.AddRaw(this, &FPreflightValidator::ValidatePreflights);
 	}
 
+	// If horde connections fails in a way that cannot be recovered, fail the validator
+	if (!HordeConnectionFailedHandler.IsValid())
+	{
+		HordeConnectionFailedHandler = PreflightService->OnHordeConnectionFailed.AddLambda([this](){ ValidationFinished(false); });
+	}
+
 	if(!TagUpdateHandler.IsValid())
 	{
 		TagUpdateHandler = PreflightTag->OnTagUpdated.AddLambda([this, PreflightService](const FTag&){ ValidatePreflights(PreflightService->GetPreflightData(), PreflightService->GetUnlinkedPreflights()); });
@@ -233,12 +239,18 @@ void FPreflightValidator::Skip()
 
 void FPreflightValidator::RemoveCallbacks()
 {
-	if(PreflightUpdateHandler.IsValid())
+	if (PreflightUpdateHandler.IsValid())
 	{
 		ServiceProvider.Pin()->GetService<FPreflightService>()->OnPreflightDataUpdated.Remove(PreflightUpdateHandler);
 		PreflightUpdateHandler.Reset();
 	}
 
+	if (HordeConnectionFailedHandler.IsValid())
+	{
+		ServiceProvider.Pin()->GetService<FPreflightService>()->OnHordeConnectionFailed.Remove(HordeConnectionFailedHandler);
+		HordeConnectionFailedHandler.Reset();
+	}
+
 	if(TagUpdateHandler.IsValid())
 	{
 		PreflightTag->OnTagUpdated.Remove(TagUpdateHandler);
diff --git a/Engine/Source/Programs/SubmitTool/Private/Logic/Validators/PreflightValidator.h b/Engine/Source/Programs/SubmitTool/Private/Logic/Validators/PreflightValidator.h
--- a/Engine/Source/Programs/SubmitTool/Private/Logic/Validators/PreflightValidator.h
+++ b/Engine/Source/Programs/SubmitTool/Private/Logic/Validators/PreflightValidator.h
@@ -28,5 +28,6 @@ protected:
 
 	virtual void ValidationFinished(bool bSuccess) override;
 	FDelegateHandle PreflightUpdateHandler;
+	FDelegateHandle HordeConnectionFailedHandler;
 	FDelegateHandle TagUpdateHandler;
 };
\ No newline at end of file
