Added Checkboxes in JIra widget

diff --git a/Engine/Source/Programs/SubmitTool/Private/View/Widgets/SJiraWidget.cpp b/Engine/Source/Programs/SubmitTool/Private/View/Widgets/SJiraWidget.cpp
--- a/Engine/Source/Programs/SubmitTool/Private/View/Widgets/SJiraWidget.cpp
+++ b/Engine/Source/Programs/SubmitTool/Private/View/Widgets/SJiraWidget.cpp
@@ -22,6 +22,7 @@
 
 namespace JiraIssuesColumns
 {
+	static const FName Checked("Checked");
 	static const FName IssueType("IssueType");
 	static const FName Issue("Issue");
 	static const FName Status("Status");
@@ -280,7 +281,11 @@ TSharedPtr<SWidget> SJiraWidget::BuildIssuesWidget()
 					.HeaderRow
 					(
 						SNew(SHeaderRow)
-						
+						+SHeaderRow::Column(JiraIssuesColumns::Checked)
+						.DefaultLabel(NSLOCTEXT("JiraWindow", "Checked", " "))
+						.ManualWidth(24)
+						.SortMode(this, &SJiraWidget::GetSortMode, JiraIssuesColumns::Checked)
+						.OnSort(this, &SJiraWidget::OnColumnSort)
 						+ SHeaderRow::Column(JiraIssuesColumns::Issue)
 						.DefaultLabel(NSLOCTEXT("JiraWindow", "Issue", "Issue"))
 						.ManualWidth(96)
@@ -450,11 +455,15 @@ class SJiraIssueNode : public SMultiColumnTableRow<TSharedRef<FJiraIssue>>
 {
 public:
 	SLATE_BEGIN_ARGS(SJiraIssueNode) {}
+		SLATE_ARGUMENT(const FTag*, TargetTag)
+		SLATE_ARGUMENT(SJiraWidget*, ParentWidget)
 	SLATE_END_ARGS()
 
 	void Construct(const FArguments& InArgs, const TSharedRef<STableViewBase>& InOwnerTable, TSharedPtr<FJiraIssue> InNode)
 	{
 		Node = InNode;
+		TargetTag = InArgs._TargetTag;
+		ParentWidget = InArgs._ParentWidget;
 
 		SMultiColumnTableRow<TSharedRef<FJiraIssue>>::Construct(
 			FSuperRowType::FArguments()
@@ -467,7 +476,21 @@ public:
 	{
 		TSharedRef<SBorder> box = SNew(SBorder).ToolTip(SNew(SToolTip).Text(FText::FromString(Node->Description)));
 		
-		if(InColumnName == JiraIssuesColumns::IssueType)
+		if (InColumnName == JiraIssuesColumns::Checked)
+		{
+			box->SetContent(SNew(SCheckBox)
+			.OnCheckStateChanged_Lambda([this](ECheckBoxState InCheckState)
+			{
+				ParentWidget->ProcessIssueSelected(Node);
+			})
+			.IsChecked_Lambda([this]
+			{
+				bool InTag = TargetTag->GetValues().ContainsByPredicate([this](const FString& Value) { return Value.Equals(Node->Key, ESearchCase::IgnoreCase); });
+
+				return InTag ? ECheckBoxState::Checked : ECheckBoxState::Unchecked;
+			}));
+		}
+		else if(InColumnName == JiraIssuesColumns::IssueType)
 		{
 			box->SetContent(SNew(STextBlock)
 				.Text(FText::FromString(Node->IssueType)));
@@ -498,11 +521,13 @@ public:
 	}
 
 	TSharedPtr<FJiraIssue> Node;
+	const FTag* TargetTag;
+	SJiraWidget* ParentWidget;
 };
 
 TSharedRef<ITableRow> SJiraWidget::GenerateRow(TSharedPtr<FJiraIssue> InItem, const TSharedRef<STableViewBase>& OwnerTable)
 {
-	return SNew(SJiraIssueNode, OwnerTable, InItem);
+	return SNew(SJiraIssueNode, OwnerTable, InItem).TargetTag(TargetTag).ParentWidget(this);
 }
 
 void SJiraWidget::OnColumnSort(EColumnSortPriority::Type InSortPriority, const FName& InColumnId, EColumnSortMode::Type InSortDirection)
diff --git a/Engine/Source/Programs/SubmitTool/Private/View/Widgets/SubmitToolWidget.cpp b/Engine/Source/Programs/SubmitTool/Private/View/Widgets/SubmitToolWidget.cpp
--- a/Engine/Source/Programs/SubmitTool/Private/View/Widgets/SubmitToolWidget.cpp
+++ b/Engine/Source/Programs/SubmitTool/Private/View/Widgets/SubmitToolWidget.cpp
@@ -96,7 +96,7 @@ void SubmitToolWidget::Construct(const FArguments& InArgs)
 	/**** Tags ****/
 	Contents->AddSlot()
 		.AutoHeight()
-		.MaxHeight(150)
+		.MaxHeight(147)
 		.Padding(0,0,0,3)
 		[
 			SNew(STagSectionWidget)
