Removing TargetType from UBT compilation of programs

diff --git a/Engine/Source/Programs/SubmitTool/Private/Logic/Validators/UBTValidator.cpp b/Engine/Source/Programs/SubmitTool/Private/Logic/Validators/UBTValidator.cpp
--- a/Engine/Source/Programs/SubmitTool/Private/Logic/Validators/UBTValidator.cpp
+++ b/Engine/Source/Programs/SubmitTool/Private/Logic/Validators/UBTValidator.cpp
@@ -120,10 +120,9 @@ bool FUBTValidator::Validate(const FString& InCLDescription, const TArray<FSourc
 		{
 			FString ProgramFileListPath = CreateFileList(ProgramFileList.Value, SubmitToolIntermediateDir + TEXT("SubmitTool/FileLists/"));
 
-			FString TargetStr = FString::Printf(TEXT("%s %s -TargetType=%s %s %s"),
+			FString TargetStr = FString::Printf(TEXT("%s %s %s %s"),
 				*ProgramFileList.Key,
 				*Platform,
-				*Target,
 				*Configuration,
 				*Arguments);
 
