Fix Submit Tool not being able to identify the CL contents changed when there are the same number of pending/shelved files

diff --git a/Engine/Source/Programs/SubmitTool/Private/Logic/ChangelistService.cpp b/Engine/Source/Programs/SubmitTool/Private/Logic/ChangelistService.cpp
--- a/Engine/Source/Programs/SubmitTool/Private/Logic/ChangelistService.cpp
+++ b/Engine/Source/Programs/SubmitTool/Private/Logic/ChangelistService.cpp
@@ -188,6 +188,7 @@ void FChangelistService::FetchChangelistCallback(const FSourceControlOperationRe
 		}
 		else
 		{
+			ChangelistState = CachedSCCProvider->GetState(ChangelistPtr.ToSharedRef(), EStateCacheUsage::Use);
 			RehydrateDataFromP4State();
 		}
 	}
@@ -258,7 +259,7 @@ void FChangelistService::RehydrateDataFromP4State()
 		CLDescription = ChangelistState->GetDescriptionText().ToString();
 	}
 
-	if(ChangelistState->GetFilesStatesNum() != FilesInCL.Num())
+	if(ChangelistState->GetFilesStates() != FilesInCL)
 	{
 		ChangeType |= ETaskArea::LocalFiles;
 		FilesInCL = ChangelistState->GetFilesStates();
@@ -267,7 +268,7 @@ void FChangelistService::RehydrateDataFromP4State()
 		PrintFilesAndShelvedFiles();
 	}
 
-	if(ChangelistState->GetShelvedFilesStatesNum() != ShelvedFilesInCL.Num())
+	if(ChangelistState->GetShelvedFilesStates() != ShelvedFilesInCL)
 	{
 		ChangeType |= ETaskArea::ShelvedFiles;
 		ShelvedFilesInCL = ChangelistState->GetShelvedFilesStates();
