Fix for submit tool tag parsing incorrectly identifying tag ids followed by a special character like "#tag-"

diff --git a/Engine/Source/Programs/SubmitTool/Private/Models/Tag.cpp b/Engine/Source/Programs/SubmitTool/Private/Models/Tag.cpp
--- a/Engine/Source/Programs/SubmitTool/Private/Models/Tag.cpp
+++ b/Engine/Source/Programs/SubmitTool/Private/Models/Tag.cpp
@@ -43,10 +43,10 @@ FString FTag::GetFullTag() const
 bool FTag::ParseTag(const FString& source)
 {
 	// regex pattern example, replacing tag, delimiter and min/maxvalues
-	// (?:(?:\r\n|\r|\n)?TAGID)( +(?:[DELIMITERS]*(?!#)(?:[\w!"\$-\/\:-\@\[-\`\{-\~]+)){MINVALUES,MAXVALUES})?
-	// (?:(?:\r\n|\r|\n)?#jira(?!\w))( +(?:[, ]*(?!#)(?:[\w!"\$-\/\:-\@\[-\`\{-\~]+)){1,256})?
+	// (?:(?:\r\n|\r|\n)?TAGID(?= |\n|$))( +(?:[DELIMITERS]*(?!#)(?:[\w!"\$-\/\:-\@\[-\`\{-\~]+)){MINVALUES,MAXVALUES})?
+	// (?:(?:\r\n|\r|\n)?#jira(?= |\n|$))( +(?:[, ]*(?!#)(?:[\w!"\$-\/\:-\@\[-\`\{-\~]+)){1,256})?
 
-	FString regexPat = "(?:(?:\\r\\n|\\r|\\n)?" + Definition.GetTagId() + "(?!\\w))( +(?:[" + Definition.ValueDelimiter + "]*(?!#)(?:[\\w!\"\\$-\\/\\:-\\@\\[-\\`\\{-\\~]+)){" + FString::FromInt(Definition.MinValues) + "," + FString::FromInt(Definition.MaxValues) + "})?";
+	FString regexPat = "(?:(?:\\r\\n|\\r|\\n)?" + Definition.GetTagId() + "(?= |\n|$))( +(?:[" + Definition.ValueDelimiter + "]*(?!#)(?:[\\w!\"\\$-\\/\\:-\\@\\[-\\`\\{-\\~]+)){" + FString::FromInt(Definition.MinValues) + "," + FString::FromInt(Definition.MaxValues) + "})?";
 	FRegexPattern Pattern = FRegexPattern(regexPat, ERegexPatternFlags::CaseInsensitive);
 	FRegexMatcher regex = FRegexMatcher(Pattern, source);
 	bool match = regex.FindNext();
