Fixed crash due to incorrect threading when files are refreshed

diff --git a/Engine/Source/Programs/SubmitTool/Private/Models/ModelInterface.cpp b/Engine/Source/Programs/SubmitTool/Private/Models/ModelInterface.cpp
--- a/Engine/Source/Programs/SubmitTool/Private/Models/ModelInterface.cpp
+++ b/Engine/Source/Programs/SubmitTool/Private/Models/ModelInterface.cpp
@@ -11,6 +11,7 @@
 #include "Configuration/Configuration.h"
 #include "SourceControlOperations.h"
 #include "HAL/PlatformApplicationMisc.h"
+#include "Async/Async.h"
 
 #include "Logic/Validators/PreflightValidator.h"
 #include "Logic/Validators/TagValidator.h"
@@ -635,12 +636,12 @@ void FModelInterface::RefreshStateBasedOnFiles()
 					if (P4LockdownService->ArePathsInLockdown(ChangelistService->GetShelvedFilesDepotPaths(), bIsUserInAllowlist))
 					{
 						UE_LOG(LogSubmitTool, Log, TEXT("There are no local files in CL %s, Submit is disabled but you can still request an Integration with your shelved files"), *ChangelistService->GetCLID());
-						ChangeState(ESubmitToolAppState::SubmitLocked);
+						AsyncTask(ENamedThreads::GameThread, [] { ChangeState(ESubmitToolAppState::SubmitLocked); });
 					}
 					else
 					{
 						UE_LOG(LogSubmitTool, Error, TEXT("There are no files in CL %s, SUBMIT IS DISABLED"), *ChangelistService->GetCLID());
-						ChangeState(ESubmitToolAppState::Errored);
+						AsyncTask(ENamedThreads::GameThread, [] { ChangeState(ESubmitToolAppState::Errored); });
 					}
 					});
 			}
@@ -684,11 +685,11 @@ void FModelInterface::RefreshStateBasedOnFiles()
 		UE::Tasks::Launch(UE_SOURCE_LOCATION, [this] {
 			if (P4LockdownService->ArePathsInLockdown(ChangelistService->GetFilesDepotPaths(), bIsUserInAllowlist))
 			{
-				ChangeState(ESubmitToolAppState::SubmitLocked);
+				AsyncTask(ENamedThreads::GameThread, [] { ChangeState(ESubmitToolAppState::SubmitLocked); });
 			}
 			else
 			{
-				ChangeState(ESubmitToolAppState::WaitingUserInput);
+				AsyncTask(ENamedThreads::GameThread, [] { ChangeState(ESubmitToolAppState::WaitingUserInput); });
 			}
 			});
 	}
@@ -915,6 +916,13 @@ bool FModelInterface::Tick(float InDeltaTime)
 
 void FModelInterface::ChangeState(ESubmitToolAppState newState, bool bForce)
 {
+	ensure(IsInGameThread());
+	if (!IsInGameThread())
+	{
+		AsyncTask(ENamedThreads::GameThread, [newState, bForce] { ChangeState(newState,bForce); });
+		return;
+	}
+
 	const ESubmitToolAppState CurrentState = SubmitToolState;
 	if (bForce)
 	{
