Tag section is now resizable in user settings and maximized by default

diff --git a/Engine/Source/Programs/SubmitTool/Private/Models/SubmitToolUserPrefs.h b/Engine/Source/Programs/SubmitTool/Private/Models/SubmitToolUserPrefs.h
--- a/Engine/Source/Programs/SubmitTool/Private/Models/SubmitToolUserPrefs.h
+++ b/Engine/Source/Programs/SubmitTool/Private/Models/SubmitToolUserPrefs.h
@@ -33,6 +33,9 @@ struct FSubmitToolUserPrefs
 	UPROPERTY()
 	FVector2D WindowSize = FVector2D();
 	
+	UPROPERTY()
+	float TagSectionSize = FLT_MAX;
+
 	UPROPERTY()
 	float P4SectionSize = 34;
 	
diff --git a/Engine/Source/Programs/SubmitTool/Private/Version/AppVersionDefines.h b/Engine/Source/Programs/SubmitTool/Private/Version/AppVersionDefines.h
--- a/Engine/Source/Programs/SubmitTool/Private/Version/AppVersionDefines.h
+++ b/Engine/Source/Programs/SubmitTool/Private/Version/AppVersionDefines.h
@@ -21,7 +21,7 @@
 
 #define SUBMIT_TOOL_VERSION_MAJOR 5
 #define SUBMIT_TOOL_VERSION_MINOR 4
-#define SUBMIT_TOOL_VERSION_PATCH 2
+#define SUBMIT_TOOL_VERSION_PATCH 3
 
 #ifndef SUBMIT_TOOL_CHANGELIST
 #define SUBMIT_TOOL_CHANGELIST_STRING TEXT("NoCL")
diff --git a/Engine/Source/Programs/SubmitTool/Private/View/Widgets/SubmitToolWidget.cpp b/Engine/Source/Programs/SubmitTool/Private/View/Widgets/SubmitToolWidget.cpp
--- a/Engine/Source/Programs/SubmitTool/Private/View/Widgets/SubmitToolWidget.cpp
+++ b/Engine/Source/Programs/SubmitTool/Private/View/Widgets/SubmitToolWidget.cpp
@@ -17,6 +17,7 @@
 #include "Widgets/Layout/SUniformGridPanel.h"
 #include "Widgets/Layout/SSpacer.h"
 #include "Widgets/Layout/SScrollBox.h"
+#include "Widgets/Layout/SSeparator.h"
 #include "Widgets/Input/SMultiLineEditableTextBox.h"
 #include "Widgets/Views/SListView.h"
 #include "Widgets/Input/SHyperlink.h"
@@ -94,16 +95,47 @@ void SubmitToolWidget::Construct(const FArguments& InArgs)
 		];
 
 	/**** Tags ****/
+	TSharedRef<STagSectionWidget> TagSection = SNew(STagSectionWidget)
+		.ParentWindow(InArgs._ParentWindow)
+		.ModelInterface(InArgs._ModelInterface);
+
+	TSharedRef<SBox> TagSectionBox =
+		SNew(SBox)
+		.MaxDesiredHeight_Lambda([TagSection]{ return FMath::Min(TagSection->GetDesiredSize().Y, FSubmitToolUserPrefs::Get()->TagSectionSize);})
+		[
+			TagSection
+		];
+
 	Contents->AddSlot()
 		.AutoHeight()
-		.MaxHeight(147)
-		.Padding(0,0,0,3)
 		[
-			SNew(STagSectionWidget)
-			.ParentWindow(InArgs._ParentWindow)
-			.ModelInterface(InArgs._ModelInterface)
+			TagSectionBox
+		];
+
+	TSharedPtr<SBorder> ResizeBorder;
+
+	Contents->AddSlot()
+		.AutoHeight()
+		[
+			SAssignNew(ResizeBorder, SBorder)
+			.OnMouseButtonUp_Lambda([](const FGeometry& Geometry, const FPointerEvent& PointerEvent) { return FReply::Handled().ReleaseMouseCapture(); })
+			.OnMouseMove_Lambda([TagSection, this](const FGeometry& Geometry, const FPointerEvent& PointerEvent){
+				if (PointerEvent.IsMouseButtonDown(EKeys::LeftMouseButton))
+				{
+					FSubmitToolUserPrefs::Get()->TagSectionSize = FMath::Clamp(FSubmitToolUserPrefs::Get()->TagSectionSize + PointerEvent.GetCursorDelta().Y, ModelInterface->GetTagsArray().Num() == 0 ? 0 : 35, TagSection->GetDesiredSize().Y);
+					return FReply::Handled();
+				}
+
+				return FReply::Unhandled();
+			
+			})
+			[
+				SNew(SSeparator).Thickness(5)
+				.Cursor(EMouseCursor::ResizeUpDown)
+			]			
 		];
 
+	ResizeBorder->SetOnMouseButtonDown(FPointerEventHandler::CreateLambda([ResizeBorder](const FGeometry& Geometry, const FPointerEvent& PointerEvent) { return FReply::Handled().CaptureMouse(ResizeBorder.ToSharedRef()); }));
 	TSharedPtr<SSplitter> Splitter;
 
 	Contents->AddSlot()
