Unshelved from pending changelist '40240962':

Create a new Validator Base Async class
Used it in shader, JSON and PackageData validators

diff --git a/Engine/Source/Programs/SubmitTool/Private/Logic/Validators/JsonValidator.cpp b/Engine/Source/Programs/SubmitTool/Private/Logic/Validators/JsonValidator.cpp
--- a/Engine/Source/Programs/SubmitTool/Private/Logic/Validators/JsonValidator.cpp
+++ b/Engine/Source/Programs/SubmitTool/Private/Logic/Validators/JsonValidator.cpp
@@ -11,126 +11,85 @@
 #include "Models/ModelInterface.h"
 
 FJsonValidator::FJsonValidator(const FName& InNameId, const FSubmitToolParameters& InParameters, TSharedRef<FSubmitToolServiceProvider> InServiceProvider, const FString& InDefinition) :
-	FValidatorBase(InNameId, InParameters, InServiceProvider, InDefinition)
+	FValidatorBaseAsync(InNameId, InParameters, InServiceProvider, InDefinition)
 {
 	ParseDefinition(InDefinition);
 }
 
-bool FJsonValidator::Validate(const FString& InCLDescription, const TArray<FSourceControlStateRef>& InFilteredFilesInCL, const TArray<const FTag*>& InTags)
+void FJsonValidator::StartAsyncWork(const FString& CLDescription, const TArray<FSourceControlStateRef>& FilteredFilesInCL, const TArray<const FTag*>& Tags)
 {
-	TArray<FString> JsonFilenames;
-	for (const FSourceControlStateRef& File : InFilteredFilesInCL)
-	{
-		JsonFilenames.Add(File->GetFilename());
-	}
+	TArray<TArray<FSourceControlStateRef>> Batches;
+	Batches.Add(TArray<FSourceControlStateRef>());
 
-	CurrentJob = MakeUnique<FJob>();
-
-	CurrentJob->Task = UE::Tasks::Launch(UE_SOURCE_LOCATION, [this, JsonFilenames, &CancelToken = CurrentJob->CancelToken]() -> bool
+	const int32 MaxTasks = FMath::Max(1, FilteredFilesInCL.Num() / 4);
+	for (const FSourceControlStateRef& File : FilteredFilesInCL)
+	{
+		if (Batches.Last().Num() >= MaxTasks)
 		{
-			bool bValid = true;
-
-			for (const FString& Filename : JsonFilenames)
-			{
-				if (CancelToken.IsCanceled())
-				{
-					break;
-				}
-
-				FString JSON;
-				const FString& RegexFiltering = GetTypedDefinition<FJsonValidatorDefinition>()->RegexLineExclusion;
-				if (!RegexFiltering.IsEmpty())
-				{
-					FRegexPattern Pattern = FRegexPattern(RegexFiltering, ERegexPatternFlags::CaseInsensitive);
-
-					TArray<FString> Lines;
-					if (!FFileHelper::LoadFileToStringArrayWithPredicate(Lines, *Filename, [&Pattern](const FString& Line) { return !FRegexMatcher(Pattern, Line).FindNext();}))
-					{
-						LogFailure(FString::Printf(TEXT("[%s] %s could not be loaded"), *GetValidatorName(), *Filename));
-						bValid = false;
-						continue;
-					}
-
-					JSON = FString::Join(Lines, TEXT("\n"));
-				}
-				else
-				{
-					if (!FFileHelper::LoadFileToString(JSON, *Filename))
-					{
-						LogFailure(FString::Printf(TEXT("[%s] %s could not be loaded"), *GetValidatorName(), *Filename));
-						bValid = false;
-						continue;
-					}
-				}
-
-				TSharedRef<TJsonReader<>> JsonReader = TJsonReaderFactory<>::Create(MoveTemp(JSON));
-				TSharedPtr<FJsonValue> JsonValue;
-				if (!FJsonSerializer::Deserialize(JsonReader, JsonValue) || !JsonValue.IsValid())
-				{
-					const FString& Error = JsonReader->GetErrorMessage();
-					LogFailure(FString::Printf(TEXT("[%s] %s is an invalid JSON file: %s"), *GetValidatorName(), *Filename, *Error));
-					bValid = false;
-				}
-			}
-
-			return bValid;
-		});
+			Batches.Add(TArray<FSourceControlStateRef>());
+		}
 
-	return true;
-}
-
-void FJsonValidator::ParseDefinition(const FString& InDefinition)
-{
-	FStringOutputDevice Errors;
-
-	FJsonValidatorDefinition* DefinitionToLoad = new FJsonValidatorDefinition;
-	FJsonValidatorDefinition::StaticStruct()->ImportText(*InDefinition, DefinitionToLoad, nullptr, 0, &Errors, FJsonValidatorDefinition::StaticStruct()->GetName());
-
-	Definition.Reset(DefinitionToLoad);
+		Batches.Last().Add(File);
+	}
 
-	if (!Errors.IsEmpty())
+	for (const TArray<FSourceControlStateRef>& Batch : Batches)
 	{
-		UE_LOG(LogSubmitTool, Error, TEXT("[%s] Error loading parameter file %s"), *GetValidatorName(), *Errors);
-		FModelInterface::SetErrorState();
+		this->StartAsyncTask([this, Batch](const UE::Tasks::FCancellationToken& InCancellationToken) -> bool
+			{
+				return ValidateJson(Batch, InCancellationToken);
+			});
 	}
 }
 
-void FJsonValidator::Tick(float InDeltaTime)
+bool FJsonValidator::ValidateJson(const TArray<FSourceControlStateRef>& InFilteredFilesInCL, const UE::Tasks::FCancellationToken& InCancellationToken) const
 {
-	FValidatorBase::Tick(InDeltaTime);
+	bool bValid = true;
 
-	if (HasWorkPending())
+	for (const FSourceControlStateRef& File : InFilteredFilesInCL)
 	{
-		return;
-	}
+		if (InCancellationToken.IsCanceled())
+		{
+			break;
+		}
 
-	const bool bAllValidationsSuccessful = !(CurrentJob && CurrentJob->Task.GetResult() == false);
-	ValidationFinished(bAllValidationsSuccessful);
+		const FString& Filename = File->GetFilename();
 
-	CurrentJob.Reset();
-}
+		FString JSON;
+		const FString& RegexFiltering = GetTypedDefinition<FJsonValidatorDefinition>()->RegexLineExclusion;
+		if (!RegexFiltering.IsEmpty())
+		{
+			FRegexPattern Pattern = FRegexPattern(RegexFiltering, ERegexPatternFlags::CaseInsensitive);
 
-void FJsonValidator::StopInternalValidations()
-{
-	if (CurrentJob)
-	{
-		// First, signal that the validation has been canceled to all running jobs.
-		CurrentJob->CancelToken.Cancel();
+			TArray<FString> Lines;
+			if (!FFileHelper::LoadFileToStringArrayWithPredicate(Lines, *Filename, [&Pattern](const FString& Line) { return !FRegexMatcher(Pattern, Line).FindNext(); }))
+			{
+				LogFailure(FString::Printf(TEXT("[%s] %s could not be loaded"), *GetValidatorName(), *Filename));
+				bValid = false;
+				continue;
+			}
 
-		// Secondly, wait on any job that is still running until their cancellation is complete.
-		CurrentJob->Task.Wait();
+			JSON = FString::Join(Lines, TEXT("\n"));
+		}
+		else
+		{
+			if (!FFileHelper::LoadFileToString(JSON, *Filename))
+			{
+				LogFailure(FString::Printf(TEXT("[%s] %s could not be loaded"), *GetValidatorName(), *Filename));
+				bValid = false;
+				continue;
+			}
+		}
 
-		// Lastly we can clean up the jobs.
-		CurrentJob.Reset();
+		TSharedRef<TJsonReader<>> JsonReader = TJsonReaderFactory<>::Create(MoveTemp(JSON));
+		TSharedPtr<FJsonValue> JsonValue;
+		if (!FJsonSerializer::Deserialize(JsonReader, JsonValue) || !JsonValue.IsValid())
+		{
+			const FString& Error = JsonReader->GetErrorMessage();
+			LogFailure(FString::Printf(TEXT("[%s] %s is an invalid JSON file: %s"), *GetValidatorName(), *Filename, *Error));
+			bValid = false;
+		}
 	}
-}
 
-bool FJsonValidator::HasWorkPending() const
-{
-	if (CurrentJob && CurrentJob->Task.IsCompleted() == false)
-	{
-		return true;
-	}
-	return false;
+	return bValid;
 }
 
diff --git a/Engine/Source/Programs/SubmitTool/Private/Logic/Validators/JsonValidator.h b/Engine/Source/Programs/SubmitTool/Private/Logic/Validators/JsonValidator.h
--- a/Engine/Source/Programs/SubmitTool/Private/Logic/Validators/JsonValidator.h
+++ b/Engine/Source/Programs/SubmitTool/Private/Logic/Validators/JsonValidator.h
@@ -2,34 +2,22 @@
 
 #pragma once
 
-#include "ValidatorBase.h"
+#include "ValidatorBaseAsync.h"
 #include "Tasks/Task.h"
 
-class FJsonValidator final : public FValidatorBase
+class FJsonValidator final : public FValidatorBaseAsync
 {
 public:
 	FJsonValidator(const FName& InNameId, const FSubmitToolParameters& InParameters, TSharedRef<FSubmitToolServiceProvider> InServiceProvider, const FString& InDefinition);
-	virtual bool Validate(const FString& InCLDescription, const TArray<FSourceControlStateRef>& InFilteredFilesInCL, const TArray<const FTag*>& InTags) override;
 
-	virtual const FString& GetValidatorTypeName() const override
-	{
-		return SubmitToolParseConstants::JsonValidator;
-	}
+	virtual void StartAsyncWork(const FString& CLDescription, const TArray<FSourceControlStateRef>& FilteredFilesInCL, const TArray<const FTag*>& Tags) override;
+
 protected:
+	bool ValidateJson(const TArray<FSourceControlStateRef>& InFilteredFilesInCL, const UE::Tasks::FCancellationToken& InCancellationToken) const;
 	virtual void ParseDefinition(const FString& InDefinition) override;
 
-private:
-	virtual void Tick(float InDeltaTime) override;
-	virtual void StopInternalValidations() override;
-
-	bool HasWorkPending() const;
-
-	struct FJob
+	virtual const FString& GetValidatorTypeName() const override
 	{
-		UE::Tasks::TTask<bool> Task;
-		UE::Tasks::FCancellationToken CancelToken;
-	};
-
-	TUniquePtr<FJob> CurrentJob;
-
+		return SubmitToolParseConstants::JsonValidator;
+	}
 };
diff --git a/Engine/Source/Programs/SubmitTool/Private/Logic/Validators/PackageDataValidator.cpp b/Engine/Source/Programs/SubmitTool/Private/Logic/Validators/PackageDataValidator.cpp
--- a/Engine/Source/Programs/SubmitTool/Private/Logic/Validators/PackageDataValidator.cpp
+++ b/Engine/Source/Programs/SubmitTool/Private/Logic/Validators/PackageDataValidator.cpp
@@ -7,7 +7,7 @@
 #include "VirtualizationUtilities.h"
 
 FPackageDataValidator::FPackageDataValidator(const FName& InNameId, const FSubmitToolParameters& InParameters, TSharedRef<FSubmitToolServiceProvider> InServiceProvider, const FString& InDefinition)
-	: FValidatorBase(InNameId, InParameters, InServiceProvider, InDefinition)
+	: FValidatorBaseAsync(InNameId, InParameters, InServiceProvider, InDefinition)
 {
 	ParseDefinition(InDefinition);
 }
@@ -28,132 +28,46 @@ void FPackageDataValidator::ParseDefinition(const FString& InDefinition)
 	}
 }
 
-bool FPackageDataValidator::Validate(const FString& InCLDescription, const TArray<FSourceControlStateRef>& InFilteredFilesInCL, const TArray<const FTag*>& InTags)
+void FPackageDataValidator::StartAsyncWork(const FString& InCLDescription, const TArray<FSourceControlStateRef>& InFilteredFilesInCL, const TArray<const FTag*>& InTags)
 {
-	if (CurrentJobs.IsEmpty() == false) // Could consider changing this to an assert when we are confident that it won't trigger
+	this->StartAsyncTask([this, InFilteredFilesInCL](const UE::Tasks::FCancellationToken& InCancellationToken) -> bool
 	{
-		LogFailure(FString::Printf(TEXT("[%s] attempting to run a new validation while a previous validation pass is still running!"), *GetValidatorName()));
-		return false;
-	}
-
-	CheckForVirtualizedContent(InFilteredFilesInCL); // Launches a single job
+		const FPackageDataValidatorDefinition* TypedDefinition = GetTypedDefinition<FPackageDataValidatorDefinition>();
+		check(TypedDefinition != nullptr);
 
-	// If we have no running work then the validation is complete
-	if (CurrentJobs.IsEmpty())
-	{
-		ValidationFinished(true);
-	}
+		bool bValid = true;
 
-	return true;
-}
-
-void FPackageDataValidator::CheckForVirtualizedContent(const TArray<FSourceControlStateRef>& InFilteredFilesInCL)
-{
-	TUniquePtr<FJob> Job = MakeUnique<FJob>();
-
-	const FPackageDataValidatorDefinition* TypedDefinition = GetTypedDefinition<FPackageDataValidatorDefinition>();
-	check(TypedDefinition != nullptr);
-
-	// We could consider launching a task per entry in ExcludedExtensions if this validator ever becomes a bottleneck
-	Job->Task = UE::Tasks::Launch(UE_SOURCE_LOCATION, [this, InFilteredFilesInCL, ExcludedExtensions = TypedDefinition->ExcludedExtensions, &CancelToken = Job->CancelToken]() -> bool
+		for (const FSourceControlStateRef& FileState : InFilteredFilesInCL)
 		{
-			bool bValid = true;
-
-			for (const FSourceControlStateRef& FileState : InFilteredFilesInCL)
+			for (const FString& Extension : TypedDefinition->ExcludedExtensions)
 			{
-				for(const FString& Extension : ExcludedExtensions)
+				if (InCancellationToken.IsCanceled())
 				{
-					if (CancelToken.IsCanceled())
-					{
-						break;
-					}
+					break;
+				}
 
-					FString ProjectFilePath;
-					FString PluginFilePath;
+				FString ProjectFilePath;
+				FString PluginFilePath;
 
-					if (UE::Virtualization::Utils::TryFindProject(FileState->GetFilename(), Extension, ProjectFilePath, PluginFilePath))
+				if (UE::Virtualization::Utils::TryFindProject(FileState->GetFilename(), Extension, ProjectFilePath, PluginFilePath))
+				{
+					// Note we cannot use the asset registry in the SubmitTool as without the proper mount points we can
+					// only work with absolute file paths so we need to load the package trailer to tell if the package
+					// has virtualized data or not.
+					UE::FPackageTrailer Trailer;
+					if (UE::FPackageTrailer::TryLoadFromFile(FileState->GetFilename(), Trailer))
 					{
-						// Note we cannot use the asset registry in the SubmitTool as without the proper mount points we can
-						// only work with absolute file paths so we need to load the package trailer to tell if the package
-						// has virtualized data or not.
-						UE::FPackageTrailer Trailer;
-						if (UE::FPackageTrailer::TryLoadFromFile(FileState->GetFilename(), Trailer))
+						if (Trailer.GetNumPayloads(UE::EPayloadStorageType::Virtualized) > 0)
 						{
-							if (Trailer.GetNumPayloads(UE::EPayloadStorageType::Virtualized) > 0)
-							{
-								LogFailure(FString::Printf(TEXT("[%s] %s has virtualized content and will not work for a '%s' project, please rehydrate!"), *GetValidatorName(), *FileState->GetFilename(), *Extension));
-								bValid = false;
-							}
+							LogFailure(FString::Printf(TEXT("[%s] %s has virtualized content and will not work for a '%s' project, please rehydrate!"), *GetValidatorName(), *FileState->GetFilename(), *Extension));
+							bValid = false;
 						}
 					}
 				}
 			}
-
-			return bValid;
-		});
-
-	CurrentJobs.Add(MoveTemp(Job));
-}
-
-void FPackageDataValidator::Tick(float InDeltaTime)
-{
-	FValidatorBase::Tick(InDeltaTime);
-
-	if (HasWorkPending())
-	{
-		return;
-	}
-
-	bool bAllValidationsSuccessful = true;
-
-	for (const TUniquePtr<FJob>& Job : CurrentJobs)
-	{
-		if (Job && Job->Task.GetResult() == false)
-		{
-			
-			bAllValidationsSuccessful = false;
-		}
-	}
-
-	ValidationFinished(bAllValidationsSuccessful);
-
-	CurrentJobs.Empty();
-}
-
-void FPackageDataValidator::StopInternalValidations()
-{
-	// First, signal that the validation has been canceled to all running jobs.
-	for (const TUniquePtr<FJob>& Job : CurrentJobs)
-	{
-		if (Job)
-		{
-			Job->CancelToken.Cancel();
-		}
-	}
-
-	// Secondly, wait on any job that is still running until their cancellation is complete.
-	for (const TUniquePtr<FJob>& Job : CurrentJobs)
-	{
-		if (Job)
-		{
-			Job->Task.Wait();
-		}
-	}
-
-	// Lastly we can clean up the jobs.
-	CurrentJobs.Empty();
-}
-
-bool FPackageDataValidator::HasWorkPending() const
-{
-	for (const TUniquePtr<FJob>& Job : CurrentJobs)
-	{
-		if (Job && Job->Task.IsCompleted() == false)
-		{
-			return true;
 		}
-	}
 
-	return false;
+		return bValid;
+	});
 }
 
diff --git a/Engine/Source/Programs/SubmitTool/Private/Logic/Validators/PackageDataValidator.h b/Engine/Source/Programs/SubmitTool/Private/Logic/Validators/PackageDataValidator.h
--- a/Engine/Source/Programs/SubmitTool/Private/Logic/Validators/PackageDataValidator.h
+++ b/Engine/Source/Programs/SubmitTool/Private/Logic/Validators/PackageDataValidator.h
@@ -2,41 +2,27 @@
 
 #pragma once
 
-#include "ValidatorBase.h"
+#include "ValidatorBaseAsync.h"
 
 #include "ISourceControlChangelistState.h"
-#include "Tasks/Task.h"
 
-class FPackageDataValidator final : public FValidatorBase
+class FPackageDataValidator final : public FValidatorBaseAsync
 {
 public:
-	using FValidatorBase::FValidatorBase;
+	using FValidatorBaseAsync::FValidatorBaseAsync;
 
 	FPackageDataValidator(const FName& InNameId, const FSubmitToolParameters& InParameters, TSharedRef<FSubmitToolServiceProvider> InServiceProvider, const FString& InDefinition);
 	virtual ~FPackageDataValidator() = default;
 
+	virtual void StartAsyncWork(const FString& InCLDescription, const TArray<FSourceControlStateRef>& InFilteredFilesInCL, const TArray<const FTag*>& InTags) override;
 private:
 
 	virtual void ParseDefinition(const FString& InDefinition) override;
 
-	virtual bool Validate(const FString& InCLDescription, const TArray<FSourceControlStateRef>& InFilteredFilesInCL, const TArray<const FTag*>& InTags) override;
-	virtual void Tick(float InDeltaTime) override;
-	virtual void StopInternalValidations() override;
+	
 
 	virtual const FString& GetValidatorTypeName() const override
 	{
 		return SubmitToolParseConstants::PackageDataValidator;
 	}
-
-	bool HasWorkPending() const;
-
-	void CheckForVirtualizedContent(const TArray<FSourceControlStateRef>& InFilteredFilesInCL);
-
-	struct FJob
-	{
-		UE::Tasks::TTask<bool> Task;
-		UE::Tasks::FCancellationToken CancelToken;
-	};
-
-	TArray<TUniquePtr<FJob>> CurrentJobs;
 };
diff --git a/Engine/Source/Programs/SubmitTool/Private/Logic/Validators/ShaderValidator.cpp b/Engine/Source/Programs/SubmitTool/Private/Logic/Validators/ShaderValidator.cpp
--- a/Engine/Source/Programs/SubmitTool/Private/Logic/Validators/ShaderValidator.cpp
+++ b/Engine/Source/Programs/SubmitTool/Private/Logic/Validators/ShaderValidator.cpp
@@ -15,96 +15,46 @@
 
 
 FShaderValidator::FShaderValidator(const FName& InNameId, const FSubmitToolParameters& InParameters, TSharedRef<FSubmitToolServiceProvider> InServiceProvider, const FString& InDefinition) :
-	FValidatorBase(InNameId, InParameters, InServiceProvider, InDefinition)
+	FValidatorBaseAsync(InNameId, InParameters, InServiceProvider, InDefinition)
 {
 }
 
-bool FShaderValidator::Validate(const FString& InCLDescription, const TArray<FSourceControlStateRef>& InFilteredFilesInCL, const TArray<const FTag*>& InTags)
+void FShaderValidator::StartAsyncWork(const FString& InCLDescription, const TArray<FSourceControlStateRef>& InFilteredFilesInCL, const TArray<const FTag*>& InTags)
 {
-	// Get the list of files to validate. This is already filtered in [Validator.Shader]:IncludeFilesWithExtension and [Validator.Shader]:IncludeFilesInDirectory.
-	TArray<FString> ShaderFilenames;
-	for(const FSourceControlStateRef& File : InFilteredFilesInCL)
+	this->StartAsyncTask([this, InFilteredFilesInCL](const UE::Tasks::FCancellationToken& InCancellationToken) -> bool
 	{
-		ShaderFilenames.Add(File->GetFilename());
-	}
+		bool bValid = true;
 
-	// We could consider launching a task per entry in ExcludedExtensions if this validator ever becomes a bottleneck
-	CurrentJob = MakeUnique<FJob>();
-
-	CurrentJob->Task = UE::Tasks::Launch(UE_SOURCE_LOCATION, [this, ShaderFilenames, &CancelToken = CurrentJob->CancelToken]() -> bool
+		for (const FSourceControlStateRef& File : InFilteredFilesInCL)
 		{
-			bool bValid = true;
-
-			for (const FString& Filename : ShaderFilenames)
+			if (InCancellationToken.IsCanceled())
 			{
-				if (CancelToken.IsCanceled())
-				{
-					break;
-				}
+				break;
+			}
 
-				FString ShaderSourceCode;
-				if (!FFileHelper::LoadFileToString(ShaderSourceCode, *Filename))
-				{
-					LogFailure(FString::Printf(TEXT("[%s] %s could not be loaded"), *GetValidatorName(), *Filename));
-					bValid = false;
-					continue;
-				}
+			const FString& Filename = File->GetFilename();
 
-				TArray<FString> Errors;
-				if (!UE::ShaderCompilerCommon::ValidateShaderAgainstKnownIssues(ShaderSourceCode, Errors, *Filename))
+			FString ShaderSourceCode;
+			if (!FFileHelper::LoadFileToString(ShaderSourceCode, *Filename))
+			{
+				LogFailure(FString::Printf(TEXT("[%s] %s could not be loaded"), *GetValidatorName(), *Filename));
+				bValid = false;
+				continue;
+			}
+
+			TArray<FString> Errors;
+			if (!UE::ShaderCompilerCommon::ValidateShaderAgainstKnownIssues(ShaderSourceCode, Errors, *Filename))
+			{
+				TStringBuilder<4096> ErrorStringBuilder;
+				for (const FString& Error : Errors)
 				{
-					TStringBuilder<4096> ErrorStringBuilder;
-					for (const FString& Error : Errors)
-					{
-						ErrorStringBuilder.Appendf(TEXT("\n%s"), *Error);
-					}
-					LogFailure(FString::Printf(TEXT("[%s] %s has validation errors:%s"), *GetValidatorName(), *Filename, ErrorStringBuilder.ToString()));
-					bValid = false;
+					ErrorStringBuilder.Appendf(TEXT("\n%s"), *Error);
 				}
+				LogFailure(FString::Printf(TEXT("[%s] %s has validation errors:%s"), *GetValidatorName(), *Filename, ErrorStringBuilder.ToString()));
+				bValid = false;
 			}
+		}
 
-			return bValid;
-		});
-
-	return true;
-}
-
-void FShaderValidator::Tick(float InDeltaTime)
-{
-	FValidatorBase::Tick(InDeltaTime);
-
-	if (HasWorkPending())
-	{
-		return;
-	}
-
-	const bool bAllValidationsSuccessful = !(CurrentJob && CurrentJob->Task.GetResult() == false);
-	ValidationFinished(bAllValidationsSuccessful);
-
-	CurrentJob.Reset();
+		return bValid;
+	});
 }
-
-void FShaderValidator::StopInternalValidations()
-{
-	if (CurrentJob)
-	{
-		// First, signal that the validation has been canceled to all running jobs.
-		CurrentJob->CancelToken.Cancel();
-
-		// Secondly, wait on any job that is still running until their cancellation is complete.
-		CurrentJob->Task.Wait();
-
-		// Lastly we can clean up the jobs.
-		CurrentJob.Reset();
-	}
-}
-
-bool FShaderValidator::HasWorkPending() const
-{
-	if (CurrentJob && CurrentJob->Task.IsCompleted() == false)
-	{
-		return true;
-	}
-	return false;
-}
-
diff --git a/Engine/Source/Programs/SubmitTool/Private/Logic/Validators/ShaderValidator.h b/Engine/Source/Programs/SubmitTool/Private/Logic/Validators/ShaderValidator.h
--- a/Engine/Source/Programs/SubmitTool/Private/Logic/Validators/ShaderValidator.h
+++ b/Engine/Source/Programs/SubmitTool/Private/Logic/Validators/ShaderValidator.h
@@ -2,32 +2,18 @@
 
 #pragma once
 
-#include "ValidatorBase.h"
-#include "Tasks/Task.h"
+#include "ValidatorBaseAsync.h"
 
-class FShaderValidator final : public FValidatorBase
+class FShaderValidator final : public FValidatorBaseAsync
 {
 public:		
 	FShaderValidator(const FName& InNameId, const FSubmitToolParameters& InParameters, TSharedRef<FSubmitToolServiceProvider> InServiceProvider, const FString& InDefinition);
-	virtual bool Validate(const FString& InCLDescription, const TArray<FSourceControlStateRef>& InFilteredFilesInCL, const TArray<const FTag*>& InTags) override;
+	virtual ~FShaderValidator() = default;
+
+	virtual void StartAsyncWork(const FString& InCLDescription, const TArray<FSourceControlStateRef>& InFilteredFilesInCL, const TArray<const FTag*>& InTags) override;
 
 	virtual const FString& GetValidatorTypeName() const override
 	{
 		return SubmitToolParseConstants::ShaderValidator;
 	}
-	
-private:
-	virtual void Tick(float InDeltaTime) override;
-	virtual void StopInternalValidations() override;
-
-	bool HasWorkPending() const;
-
-	struct FJob
-	{
-		UE::Tasks::TTask<bool> Task;
-		UE::Tasks::FCancellationToken CancelToken;
-	};
-
-	TUniquePtr<FJob> CurrentJob;
-
 };
diff --git a/Engine/Source/Programs/SubmitTool/Private/Logic/Validators/ValidatorBase.cpp b/Engine/Source/Programs/SubmitTool/Private/Logic/Validators/ValidatorBase.cpp
--- a/Engine/Source/Programs/SubmitTool/Private/Logic/Validators/ValidatorBase.cpp
+++ b/Engine/Source/Programs/SubmitTool/Private/Logic/Validators/ValidatorBase.cpp
@@ -227,7 +227,7 @@ void FValidatorBase::ValidationFinished(const bool bHasPassed)
 {
 	if(bHasPassed)
 	{
-		UE_LOG(LogValidatorsResult, Log, TEXT("[%s]: Task Succeeded!"), *GetValidatorName());
+		UE_LOG(LogValidatorsResult, Log, TEXT("[%s]: Task Succeeded! (%s)"), *GetValidatorName(), *FGenericPlatformTime::PrettyTime(RunTime));
 
 		if(Definition->bUsesIncrementalCache)
 		{
diff --git a/Engine/Source/Programs/SubmitTool/Private/Logic/Validators/ValidatorBaseAsync.cpp b/Engine/Source/Programs/SubmitTool/Private/Logic/Validators/ValidatorBaseAsync.cpp
new file mode 100644
--- /dev/null
+++ b/Engine/Source/Programs/SubmitTool/Private/Logic/Validators/ValidatorBaseAsync.cpp
@@ -0,0 +1,93 @@
+// Copyright Epic Games, Inc. All Rights Reserved.
+
+#include "ValidatorBaseAsync.h"
+
+FValidatorBaseAsync::FValidatorBaseAsync(const FName& InNameId, const FSubmitToolParameters& InParameters, TSharedRef<FSubmitToolServiceProvider> InServiceProvider, const FString& InDefinition):
+	FValidatorBase(InNameId, InParameters, InServiceProvider, InDefinition)
+{
+}
+
+bool FValidatorBaseAsync::Validate(const FString& CLDescription, const TArray<FSourceControlStateRef>& FilteredFilesInCL, const TArray<const FTag*>& Tags)
+{
+	if(!CurrentTasks.IsEmpty())
+	{
+		LogFailure(FString::Printf(TEXT("[%s] attempting to run a new validation while a previous validation pass is still running!"), *GetValidatorName()));
+		return false;
+	}
+
+	CancellationToken = MakeUnique<UE::Tasks::FCancellationToken>();
+
+	// Create a new token for each validation
+	this->StartAsyncWork(CLDescription, FilteredFilesInCL, Tags);
+
+	// add validation to check if some work has been done
+	if(CurrentTasks.Num() == 0)
+	{
+		return false;
+	}
+
+	return true;
+}
+
+void FValidatorBaseAsync::StartAsyncTask(TFunction<bool(const UE::Tasks::FCancellationToken&)>&& InTask)
+{
+	UE::Tasks::TTask<bool> Task = UE::Tasks::Launch(UE_SOURCE_LOCATION, [this, InTask]() -> bool
+		{
+			return InTask(*CancellationToken);
+		});
+
+	CurrentTasks.Add(Task);
+}
+
+void FValidatorBaseAsync::Tick(float InDeltaTime)
+{
+	FValidatorBase::Tick(InDeltaTime);
+
+	if (HasWorkPending())
+	{
+		return;
+	}
+
+	bool bAllValidationsSuccessful = true;
+
+	for (UE::Tasks::TTask<bool>& Task : CurrentTasks)
+	{
+		bAllValidationsSuccessful &= Task.GetResult();
+	}
+
+	ValidationFinished(bAllValidationsSuccessful);
+
+	CurrentTasks.Empty();
+}
+
+void FValidatorBaseAsync::StopInternalValidations()
+{
+	// First, signal that the validation has been canceled to all running jobs.
+	if (CancellationToken.IsValid())
+	{
+		CancellationToken->Cancel();
+	}
+	
+	// Secondly, wait on any job that is still running until their cancellation is complete.
+	for (UE::Tasks::TTask<bool>& Task : CurrentTasks)
+	{
+		// that's just for you Juan
+		Task.Wait();
+	}
+
+	// Lastly we can clean up the jobs.
+	CurrentTasks.Reset();
+}
+
+bool FValidatorBaseAsync::HasWorkPending() const
+{
+	for (const UE::Tasks::TTask<bool>& Task : CurrentTasks)
+	{
+		if (!Task.IsCompleted())
+		{
+			return true;
+		}
+	}
+
+	return false;
+}
\ No newline at end of file
diff --git a/Engine/Source/Programs/SubmitTool/Private/Logic/Validators/ValidatorBaseAsync.h b/Engine/Source/Programs/SubmitTool/Private/Logic/Validators/ValidatorBaseAsync.h
new file mode 100644
--- /dev/null
+++ b/Engine/Source/Programs/SubmitTool/Private/Logic/Validators/ValidatorBaseAsync.h
@@ -0,0 +1,33 @@
+// Copyright Epic Games, Inc. All Rights Reserved.
+
+#pragma once
+
+#include "ValidatorBase.h"
+#include "Tasks/Task.h"
+
+class FValidatorBaseAsync : public FValidatorBase
+{
+public:
+	FValidatorBaseAsync(const FName& InNameId, const FSubmitToolParameters& InParameters, TSharedRef<FSubmitToolServiceProvider> InServiceProvider, const FString& InDefinition);
+
+	virtual bool Validate(const FString& CLDescription, const TArray<FSourceControlStateRef>& FilteredFilesInCL, const TArray<const FTag*>& Tags) override final;
+
+	//
+	// Override this method to start the async work
+	// Please make sure to check that the CancellationToken has not been cancelled and stop processing in that case
+	//
+	virtual void StartAsyncWork(const FString& CLDescription, const TArray<FSourceControlStateRef>& FilteredFilesInCL, const TArray<const FTag*>& Tags) = 0;
+
+	virtual void Tick(float InDeltaTime) override  final;
+	virtual void StopInternalValidations() override final;
+
+protected:
+	void StartAsyncTask(TFunction<bool(const UE::Tasks::FCancellationToken&)>&& InTask);
+
+private:
+
+	bool HasWorkPending() const;
+
+	TUniquePtr<UE::Tasks::FCancellationToken> CancellationToken;
+	TArray<UE::Tasks::TTask<bool>> CurrentTasks;
+};
\ No newline at end of file
