Added missing definition for buildfix

diff --git a/Engine/Source/Programs/SubmitTool/Private/Logic/Validators/JsonValidator.cpp b/Engine/Source/Programs/SubmitTool/Private/Logic/Validators/JsonValidator.cpp
--- a/Engine/Source/Programs/SubmitTool/Private/Logic/Validators/JsonValidator.cpp
+++ b/Engine/Source/Programs/SubmitTool/Private/Logic/Validators/JsonValidator.cpp
@@ -93,3 +93,18 @@ bool FJsonValidator::ValidateJson(const TArray<FSourceControlStateRef>& InFilter
 	return bValid;
 }
 
+void FJsonValidator::ParseDefinition(const FString& InDefinition)
+{
+	FStringOutputDevice Errors;
+
+	FJsonValidatorDefinition* DefinitionToLoad = new FJsonValidatorDefinition;
+	FJsonValidatorDefinition::StaticStruct()->ImportText(*InDefinition, DefinitionToLoad, nullptr, 0, &Errors, FJsonValidatorDefinition::StaticStruct()->GetName());
+
+	Definition.Reset(DefinitionToLoad);
+
+	if (!Errors.IsEmpty())
+	{
+		UE_LOG(LogSubmitTool, Error, TEXT("[%s] Error loading parameter file %s"), *GetValidatorName(), *Errors);
+		FModelInterface::SetErrorState();
+	}
+}
