Added override to tag parsing so the same tag can be identified in more than one form (#tag and #tag-123)

diff --git a/Engine/Source/Programs/SubmitTool/Private/Models/Tag.cpp b/Engine/Source/Programs/SubmitTool/Private/Models/Tag.cpp
--- a/Engine/Source/Programs/SubmitTool/Private/Models/Tag.cpp
+++ b/Engine/Source/Programs/SubmitTool/Private/Models/Tag.cpp
@@ -21,7 +21,7 @@ FString FTag::GetFullTag() const
 	TStringBuilder<256> strBuilder;
 
 	strBuilder.AppendChar('\n');
-	strBuilder.Append(Definition.TagId);
+	strBuilder.Append(TagKeyValue.IsEmpty() ? Definition.TagId : TagKeyValue);
 
 	if(TagValues.Num() > 0)
 	{
@@ -46,28 +46,30 @@ bool FTag::ParseTag(const FString& source)
 	// (?:(?:\r\n|\r|\n)?TAGID(?= |\n|$))( +(?:[DELIMITERS]*(?!#)(?:[\w!"\$-\/\:-\@\[-\`\{-\~]+)){MINVALUES,MAXVALUES})?
 	// (?:(?:\r\n|\r|\n)?#jira(?= |\n|$))( +(?:[, ]*(?!#)(?:[\w!"\$-\/\:-\@\[-\`\{-\~]+)){1,256})?
 
-	FString regexPat = "(?:(?:\\r\\n|\\r|\\n)?" + Definition.GetTagId() + "(?= |\n|$))( +(?:[" + Definition.ValueDelimiter + "]*(?!#)(?:[\\w!\"\\$-\\/\\:-\\@\\[-\\`\\{-\\~]+)){" + FString::FromInt(Definition.MinValues) + "," + FString::FromInt(Definition.MaxValues) + "})?";
-	FRegexPattern Pattern = FRegexPattern(regexPat, ERegexPatternFlags::CaseInsensitive);
-	FRegexMatcher regex = FRegexMatcher(Pattern, source);
-	bool match = regex.FindNext();
-	if(match)
+	const FString TagIdParse = Definition.RegexParseOverride.IsEmpty() ? Definition.GetTagId() : Definition.RegexParseOverride;
+	const FString RegexPat = "(?:(?:\\r\\n|\\r|\\n)?(" + TagIdParse + ")(?= |\n|$))( +(?:[" + Definition.ValueDelimiter + "]*(?!#)(?:[\\w!\"\\$-\\/\\:-\\@\\[-\\`\\{-\\~]+)){" + FString::FromInt(Definition.MinValues) + "," + FString::FromInt(Definition.MaxValues) + "})?";
+	FRegexPattern Pattern = FRegexPattern(RegexPat, ERegexPatternFlags::CaseInsensitive);
+	FRegexMatcher Regex = FRegexMatcher(Pattern, source);
+	bool Match = Regex.FindNext();
+	if(Match)
 	{
 		bIsDirty = false;
-		StartPos = regex.GetMatchBeginning();
-		LastSize = regex.GetMatchEnding() - regex.GetMatchBeginning();
+		StartPos = Regex.GetMatchBeginning();
+		LastSize = Regex.GetMatchEnding() - Regex.GetMatchBeginning();
 
-		UE_LOG(LogSubmitToolDebug, Log, TEXT("Start: %d"), regex.GetMatchBeginning());
-		UE_LOG(LogSubmitToolDebug, Log, TEXT("Regex matched: %s"), *regex.GetCaptureGroup(0));
+		UE_LOG(LogSubmitToolDebug, Log, TEXT("Start: %d"), Regex.GetMatchBeginning());
+		UE_LOG(LogSubmitToolDebug, Log, TEXT("Regex matched: %s"), *Regex.GetCaptureGroup(0));
 
-		FString capture = regex.GetCaptureGroup(1).TrimStart();
-		capture.ParseIntoArray(TagValues, Delimiters.GetData(), SplittedDelims.Num());
+		TagKeyValue = Regex.GetCaptureGroup(1).TrimStartAndEnd();
+		const FString Capture = Regex.GetCaptureGroup(2).TrimStart();
+		Capture.ParseIntoArray(TagValues, Delimiters.GetData(), SplittedDelims.Num());
 
-		for(const FString& value : TagValues)
+		for(const FString& Value : TagValues)
 		{
-			UE_LOG(LogSubmitToolDebug, Log, TEXT("Captured Value: %s"), *value);
+			UE_LOG(LogSubmitToolDebug, Log, TEXT("Captured Value: %s"), *Value);
 		}
 
-		UE_LOG(LogSubmitToolDebug, Log, TEXT("End: %d"), regex.GetMatchEnding());
+		UE_LOG(LogSubmitToolDebug, Log, TEXT("End: %d"), Regex.GetMatchEnding());
 	}
 	else
 	{
@@ -79,7 +81,7 @@ bool FTag::ParseTag(const FString& source)
 	{
 		OnTagUpdated.Broadcast(*this);
 	}
-	return match;
+	return Match;
 }
 
 void FTag::SetValues(const FString& valuesText)
diff --git a/Engine/Source/Programs/SubmitTool/Private/Models/Tag.h b/Engine/Source/Programs/SubmitTool/Private/Models/Tag.h
--- a/Engine/Source/Programs/SubmitTool/Private/Models/Tag.h
+++ b/Engine/Source/Programs/SubmitTool/Private/Models/Tag.h
@@ -71,5 +71,6 @@ private:
 	TArray<TCHAR const*> Delimiters;
 	TArray<FString> SplittedDelims;
 	TArray<FString> TagValues;
+	FString TagKeyValue;
 	mutable ETagState ValidationState;
 };
diff --git a/Engine/Source/Programs/SubmitTool/Private/Models/TagDefinition.h b/Engine/Source/Programs/SubmitTool/Private/Models/TagDefinition.h
--- a/Engine/Source/Programs/SubmitTool/Private/Models/TagDefinition.h
+++ b/Engine/Source/Programs/SubmitTool/Private/Models/TagDefinition.h
@@ -48,6 +48,9 @@ struct FTagDefinition
 	UPROPERTY()
 	FString TagId;
 
+	UPROPERTY()
+	FString RegexParseOverride;
+
 	UPROPERTY()
 	FString TagLabel;
 
