Fix for review tag values not being used for the create swarm request and missing the users

diff --git a/Engine/Source/Programs/SubmitTool/Private/Logic/IntegrationService.cpp b/Engine/Source/Programs/SubmitTool/Private/Logic/IntegrationService.cpp
--- a/Engine/Source/Programs/SubmitTool/Private/Logic/IntegrationService.cpp
+++ b/Engine/Source/Programs/SubmitTool/Private/Logic/IntegrationService.cpp
@@ -135,7 +135,7 @@ void FIntegrationService::RequestIntegration(const FOnBooleanValueChanged OnComp
 							FString SwarmUrl;
 							if(!SwarmService->GetCurrentReviewUrl(SwarmUrl))
 							{
-								SwarmService->CreateReview(OnCreateReviewComplete::CreateLambda([this, ChangelistService, SourceControlService, JiraService, JiraTagValue, OnComplete](bool bSuccess, const FString& Response) {
+								SwarmService->CreateReview(SwarmService->GetUsersInSwarmTag(), OnCreateReviewComplete::CreateLambda([this, ChangelistService, SourceControlService, JiraService, JiraTagValue, OnComplete](bool bSuccess, const FString& Response) {
 									if(bSuccess)
 									{
 										JiraService->GetIssueAndCreateServiceDeskRequest(JiraTagValue, ChangelistService->GetCLDescription(), Response, SourceControlService->GetCurrentStreamName(), IntegrationOptions, OnComplete);
@@ -207,7 +207,7 @@ void FIntegrationService::RequestIntegration(const FOnBooleanValueChanged OnComp
 			}
 			else
 			{
-				SwarmService->CreateReview(OnCreateReviewComplete::CreateLambda([this, ChangelistService, SourceControlService, JiraService, JiraTagValue, OnComplete](bool bSuccess, const FString& Response) {
+				SwarmService->CreateReview(SwarmService->GetUsersInSwarmTag(), OnCreateReviewComplete::CreateLambda([this, ChangelistService, SourceControlService, JiraService, JiraTagValue, OnComplete](bool bSuccess, const FString& Response) {
 					if(bSuccess)
 					{
 						JiraService->GetIssueAndCreateServiceDeskRequest(JiraTagValue, ChangelistService->GetCLDescription(), Response, SourceControlService->GetCurrentStreamName(), IntegrationOptions, OnComplete);
diff --git a/Engine/Source/Programs/SubmitTool/Private/Logic/SwarmService.cpp b/Engine/Source/Programs/SubmitTool/Private/Logic/SwarmService.cpp
--- a/Engine/Source/Programs/SubmitTool/Private/Logic/SwarmService.cpp
+++ b/Engine/Source/Programs/SubmitTool/Private/Logic/SwarmService.cpp
@@ -91,7 +91,7 @@ void FSwarmService::FetchReview(const OnGetReviewComplete& OnComplete)
 	HttpRequest->ProcessRequest();
 }
 
-void FSwarmService::CreateReview(const OnCreateReviewComplete& OnComplete)
+void FSwarmService::CreateReview(const TArray<FString>& InReviewers, const OnCreateReviewComplete& OnComplete)
 {
 	if(!bCanConnect)
 	{
@@ -124,34 +124,28 @@ void FSwarmService::CreateReview(const OnCreateReviewComplete& OnComplete)
 	TArray<TSharedPtr<FJsonValue>> ReviewersObject;
 	TArray<TSharedPtr<FJsonValue>> GroupsObject;
 
-	for (const FTag* Tag : ServiceProvider.Pin()->GetService<FTagService>()->GetTagsArray())
+	for (const FString& value : InReviewers)
 	{
-		if (Tag->Definition.InputSubType.Equals(TEXT("Swarm"), ESearchCase::IgnoreCase))
+		if (Users != nullptr)
 		{
-			for (const FString& value : Tag->GetValues())
+			TSharedPtr<FUserData>* FoundUser = Users->FindByPredicate([value = value.TrimChar(TCHAR('@'))](TSharedPtr<FUserData>& InUserData) { return InUserData->Username == value; });
+			if (FoundUser != nullptr)
 			{
-				if (Users != nullptr)
-				{
-					TSharedPtr<FUserData>* FoundUser = Users->FindByPredicate([value = value.TrimChar(TCHAR('@'))](TSharedPtr<FUserData>& InUserData) { return InUserData->Username == value; });
-					if (FoundUser != nullptr)
-					{
-						ReviewersObject.Add(MakeShared<FJsonValueString>((*FoundUser)->Username));
-					}
-					continue;
-				}
+				ReviewersObject.Add(MakeShared<FJsonValueString>((*FoundUser)->Username));
+			}
+			continue;
+		}
 
-				if (Groups != nullptr)
-				{
-					TSharedPtr<FString>* FoundGroup = Groups->FindByPredicate([&value](TSharedPtr<FString>& InGroupData) { return *InGroupData == value; });
-					if (FoundGroup != nullptr)
-					{
-						TSharedPtr<FJsonObject> GroupObject = MakeShared<FJsonObject>();
-						TSharedPtr<FJsonObject> GroupDetails = MakeShared<FJsonObject>();
-						GroupDetails->SetStringField(TEXT("required"), TEXT("false"));
-						GroupObject->SetObjectField(*(*FoundGroup), GroupDetails);
-						GroupsObject.Add(MakeShared<FJsonValueObject>(GroupObject));
-					}
-				}
+		if (Groups != nullptr)
+		{
+			TSharedPtr<FString>* FoundGroup = Groups->FindByPredicate([&value](TSharedPtr<FString>& InGroupData) { return *InGroupData == value; });
+			if (FoundGroup != nullptr)
+			{
+				TSharedPtr<FJsonObject> GroupObject = MakeShared<FJsonObject>();
+				TSharedPtr<FJsonObject> GroupDetails = MakeShared<FJsonObject>();
+				GroupDetails->SetStringField(TEXT("required"), TEXT("false"));
+				GroupObject->SetObjectField(*(*FoundGroup), GroupDetails);
+				GroupsObject.Add(MakeShared<FJsonValueObject>(GroupObject));
 			}
 		}
 	}
@@ -319,4 +313,18 @@ const FString& FSwarmService::GetSwarmURL()
 FString FSwarmService::ReviewsURL() const
 {
 	return SwarmURL.IsEmpty() ? FString{} : SwarmURL / TEXT("api/v9/reviews");
-}
\ No newline at end of file
+}
+
+const TArray<FString> FSwarmService::GetUsersInSwarmTag() const
+{
+	TArray<FString> Reviewers;
+	
+	for (const FTag* Tag : ServiceProvider.Pin()->GetService<FTagService>()->GetTagsArray())
+	{
+		if (Tag->Definition.InputSubType.Equals(TEXT("Swarm"), ESearchCase::IgnoreCase))
+		{
+			return Tag->GetValues();
+		}
+	}
+	return TArray<FString>();
+}
diff --git a/Engine/Source/Programs/SubmitTool/Private/Logic/SwarmService.h b/Engine/Source/Programs/SubmitTool/Private/Logic/SwarmService.h
--- a/Engine/Source/Programs/SubmitTool/Private/Logic/SwarmService.h
+++ b/Engine/Source/Programs/SubmitTool/Private/Logic/SwarmService.h
@@ -21,7 +21,7 @@ public:
 	FSwarmService(TWeakPtr<FSubmitToolServiceProvider> InServiceProvider);
 	void FetchReview(const OnGetReviewComplete& OnComplete);
 
-	void CreateReview(const OnCreateReviewComplete& OnComplete);
+	void CreateReview(const TArray<FString>& InReviewers, const OnCreateReviewComplete& OnComplete);
 
 	void UpdateReviewDescription(const TDelegate<void(bool)>& OnComplete, const FString& InDescription);
 
@@ -56,10 +56,13 @@ public:
 			|| (UpdateSwarmRequest.IsValid() && UpdateSwarmRequest->GetStatus() == EHttpRequestStatus::Processing);
 	}
 
+	const TArray<FString> GetUsersInSwarmTag() const;
+
 private:
 	const FString& GetSwarmURL();
 	FString ReviewsURL() const;
 
+
 private:
 	TSharedPtr<IHttpRequest> CreateSwarmRequest;
 	TSharedPtr<IHttpRequest> UpdateSwarmRequest;
diff --git a/Engine/Source/Programs/SubmitTool/Private/Models/ModelInterface.cpp b/Engine/Source/Programs/SubmitTool/Private/Models/ModelInterface.cpp
--- a/Engine/Source/Programs/SubmitTool/Private/Models/ModelInterface.cpp
+++ b/Engine/Source/Programs/SubmitTool/Private/Models/ModelInterface.cpp
@@ -444,17 +444,17 @@ void FModelInterface::ShowSwarmReview()
 	}
 }
 
-void FModelInterface::RequestSwarmReview()
+void FModelInterface::RequestSwarmReview(const TArray<FString>& InReviewers)
 {
 	if (!HasSwarmReview() && SwarmService.IsValid())
 	{
 		if (!HasShelvedFiles())
 		{
-			ChangelistService->CreateShelvedFiles(FSourceControlOperationComplete::CreateLambda([this](const FSourceControlOperationRef& DeleteShelvedOp, ECommandResult::Type Result)
+			ChangelistService->CreateShelvedFiles(FSourceControlOperationComplete::CreateLambda([this, InReviewers](const FSourceControlOperationRef& DeleteShelvedOp, ECommandResult::Type Result)
 				{
 					if (Result == ECommandResult::Succeeded)
 					{
-						RequestSwarmReview();
+						RequestSwarmReview(InReviewers);
 					}
 					else
 					{
@@ -465,7 +465,7 @@ void FModelInterface::RequestSwarmReview()
 			return;
 		}
 
-		SwarmService->CreateReview(OnCreateReviewComplete::CreateRaw(this, &FModelInterface::OnSwarmCreateCompleted));
+		SwarmService->CreateReview(InReviewers, OnCreateReviewComplete::CreateRaw(this, &FModelInterface::OnSwarmCreateCompleted));
 	}
 }
 
diff --git a/Engine/Source/Programs/SubmitTool/Private/Models/ModelInterface.h b/Engine/Source/Programs/SubmitTool/Private/Models/ModelInterface.h
--- a/Engine/Source/Programs/SubmitTool/Private/Models/ModelInterface.h
+++ b/Engine/Source/Programs/SubmitTool/Private/Models/ModelInterface.h
@@ -181,7 +181,7 @@ public:
 
 
 	void ShowSwarmReview();
-	void RequestSwarmReview();
+	void RequestSwarmReview(const TArray<FString>& InReviewers);
 	FDelegateHandle AddPreflightUpdateCallback
 	(const FOnPreflightDataUpdated::FDelegate& Callback) const						{ return PreflightService->OnPreflightDataUpdated.Add(Callback); }
 
diff --git a/Engine/Source/Programs/SubmitTool/Private/View/Widgets/TagWidget.cpp b/Engine/Source/Programs/SubmitTool/Private/View/Widgets/TagWidget.cpp
--- a/Engine/Source/Programs/SubmitTool/Private/View/Widgets/TagWidget.cpp
+++ b/Engine/Source/Programs/SubmitTool/Private/View/Widgets/TagWidget.cpp
@@ -698,7 +698,7 @@ FReply STagWidget::OnSwarmClick()
 	}
 	else
 	{
-		ModelInterface->RequestSwarmReview();
+		ModelInterface->RequestSwarmReview(Tag->GetValues());
 	}
 
 	return FReply::Handled();
