Adding user option to automatically update the preflight reference without getting a popup

diff --git a/Engine/Source/Programs/SubmitTool/Private/Logic/DialogFactory.cpp b/Engine/Source/Programs/SubmitTool/Private/Logic/DialogFactory.cpp
--- a/Engine/Source/Programs/SubmitTool/Private/Logic/DialogFactory.cpp
+++ b/Engine/Source/Programs/SubmitTool/Private/Logic/DialogFactory.cpp
@@ -6,7 +6,7 @@
 #include "Framework/Application/SlateApplication.h"
 
 EDialogFactoryResult FDialogFactory::ShowDialog(const FText& Title, const FText& Description,
-	const TArray<FString>& AvailableButtons)
+	const TArray<FString>& AvailableButtons, const TSharedPtr<SWidget> AdditionalContent)
 {	
 	TSharedPtr<SWindow> Window = SNew(SWindow)
 	.Title(Title)
@@ -26,7 +26,8 @@ EDialogFactoryResult FDialogFactory::ShowDialog(const FText& Title, const FText&
 	const TSharedPtr<SConfirmDialogWidget> ConfirmWidget = SNew(SConfirmDialogWidget)
 		.DescriptionText(Description)
 		.Buttons(AvailableButtons)
-		.ResultCallback(ResultsCallback);
+		.ResultCallback(ResultsCallback)
+		.AdditionalContent(AdditionalContent);
 
 	Window->SetContent(ConfirmWidget.ToSharedRef());
 
@@ -35,17 +36,17 @@ EDialogFactoryResult FDialogFactory::ShowDialog(const FText& Title, const FText&
 	return ButtonClicked;
 }
 
-EDialogFactoryResult FDialogFactory::ShowConfirmDialog(const FText& Title, const FText& Description)
+EDialogFactoryResult FDialogFactory::ShowConfirmDialog(const FText& Title, const FText& Description, const TSharedPtr<SWidget> AdditionalContent)
 {
-	return ShowDialog(Title, Description, TArray<FString>{ TEXT("Confirm"), TEXT("Cancel") });
+	return ShowDialog(Title, Description, TArray<FString>{ TEXT("Confirm"), TEXT("Cancel") }, AdditionalContent);
 };
 
-EDialogFactoryResult FDialogFactory::ShowInformationDialog(const FText& Title, const FText& Description)
+EDialogFactoryResult FDialogFactory::ShowInformationDialog(const FText& Title, const FText& Description, const TSharedPtr<SWidget> AdditionalContent)
 {
-	return ShowDialog(Title, Description, TArray<FString>{ TEXT("Ok") });
+	return ShowDialog(Title, Description, TArray<FString>{ TEXT("Ok") }, AdditionalContent);
 }
 
-EDialogFactoryResult FDialogFactory::ShowYesNoDialog(const FText& Title, const FText& Description)
+EDialogFactoryResult FDialogFactory::ShowYesNoDialog(const FText& Title, const FText& Description, const TSharedPtr<SWidget> AdditionalContent)
 {
-	return ShowDialog(Title, Description, TArray<FString>{ TEXT("Yes"), TEXT("No") });
+	return ShowDialog(Title, Description, TArray<FString>{ TEXT("Yes"), TEXT("No") }, AdditionalContent);
 }
diff --git a/Engine/Source/Programs/SubmitTool/Private/Logic/DialogFactory.h b/Engine/Source/Programs/SubmitTool/Private/Logic/DialogFactory.h
--- a/Engine/Source/Programs/SubmitTool/Private/Logic/DialogFactory.h
+++ b/Engine/Source/Programs/SubmitTool/Private/Logic/DialogFactory.h
@@ -20,7 +20,7 @@ enum class EDialogFactoryResult : size_t
 class FDialogFactory
 {
 public:
-	static EDialogFactoryResult ShowDialog(const FText& Title, const FText& Description, const TArray<FString>& AvailableButtons);
+	static EDialogFactoryResult ShowDialog(const FText& Title, const FText& Description, const TArray<FString>& AvailableButtons, const TSharedPtr<SWidget> AdditionalContent = nullptr);
 
 	/**
 	 * @brief Shows a dialog screen with the Cancel and Confirm buttons
@@ -28,7 +28,7 @@ public:
 	 * @param Description The description of the dialog
 	 * @return The button idx that is pressed
 	 */
-	static EDialogFactoryResult ShowConfirmDialog(const FText& Title, const FText& Description);
+	static EDialogFactoryResult ShowConfirmDialog(const FText& Title, const FText& Description, const TSharedPtr<SWidget> AdditionalContent = nullptr);
 
 	/**
 	 * @brief Shows a dialog screen with the Confirm button
@@ -36,7 +36,7 @@ public:
 	 * @param Description The description of the dialog
 	 * @return The button idx that is pressed
 	 */
-	static EDialogFactoryResult ShowInformationDialog(const FText& Title, const FText& Description);
+	static EDialogFactoryResult ShowInformationDialog(const FText& Title, const FText& Description, const TSharedPtr<SWidget> AdditionalContent = nullptr);
 
 	/**
 	 * @brief Shows a dialog screen with the Yes and No buttons
@@ -44,5 +44,5 @@ public:
 	 * @param Description The description of the dialog
 	 * @return The button idx that is pressed
 	 */
-	static EDialogFactoryResult ShowYesNoDialog(const FText& Title, const FText& Description);
+	static EDialogFactoryResult ShowYesNoDialog(const FText& Title, const FText& Description, const TSharedPtr<SWidget> AdditionalContent = nullptr);
 };
diff --git a/Engine/Source/Programs/SubmitTool/Private/Logic/PreflightService.cpp b/Engine/Source/Programs/SubmitTool/Private/Logic/PreflightService.cpp
--- a/Engine/Source/Programs/SubmitTool/Private/Logic/PreflightService.cpp
+++ b/Engine/Source/Programs/SubmitTool/Private/Logic/PreflightService.cpp
@@ -10,6 +10,9 @@
 #include "Serialization/JsonSerializer.h"
 #include "Widgets/Docking/SDockTab.h"
 #include "Internationalization/Regex.h"
+#include "Widgets/Input/SCheckBox.h"
+#include "Widgets/Input/SButton.h"
+#include "Widgets/Text/STextBlock.h"
 
 #include "CommandLine/CmdLineParameters.h"
 #include "Configuration/Configuration.h"
@@ -18,6 +21,7 @@
 #include "Logic/TagService.h"
 #include "Logging/SubmitToolLog.h"
 #include "Models/ModelInterface.h"
+#include "Models/SubmitToolUserPrefs.h"
 #include "Telemetry/TelemetryService.h"
 #include "Parameters/SubmitToolParameters.h"
 
@@ -328,20 +332,31 @@ void FPreflightService::FetchPreflightInfo(bool bRequeue, const FString& InOAuth
 								if(HordePreflights->PreflightList.Num() != 0)
 								{
 									FString CurrentTagValue = PreflightTag->GetValuesText();
-
+									
 									if(!bStopAskingTagUpdate && !CurrentTagValue.Contains(HordePreflights->PreflightList[0].ID))
 									{
-										EDialogFactoryResult Result = ShowUpdatePreflightTagDialog();
-										if(Result == EDialogFactoryResult::FirstButton)
+										if (FSubmitToolUserPrefs::Get()->bAutoUpdatePreflight)
 										{
-											// Set the latest one as the tag value
 											ModelInterface->SetTagValues(*PreflightTag, HordePreflights->PreflightList[0].ID);
+											UE_LOG(LogSubmitTool, Log, TEXT("Tag %s has been updated with the latest associated preflight %sjob/%s"), *PreflightTag->Definition.GetTagId(), *Definition.HordeServerAddress, *HordePreflights->PreflightList[0].ID)
 											bCLDescriptionModified = true;
 										}
 										else
 										{
-											bStopAskingTagUpdate = true;
+											EDialogFactoryResult Result = ShowUpdatePreflightTagDialog();
+											if(Result == EDialogFactoryResult::FirstButton)
+											{
+												// Set the latest one as the tag value
+												ModelInterface->SetTagValues(*PreflightTag, HordePreflights->PreflightList[0].ID);
+												UE_LOG(LogSubmitTool, Log, TEXT("Tag %s has been updated with the latest associated preflight %sjob/%s"), *PreflightTag->Definition.GetTagId(), *Definition.HordeServerAddress, *HordePreflights->PreflightList[0].ID)
+												bCLDescriptionModified = true;
+											}
+											else
+											{
+												bStopAskingTagUpdate = true;
+											}
 										}
+
 									}
 								}
 
@@ -702,5 +717,33 @@ EDialogFactoryResult FPreflightService::ShowUpdatePreflightTagDialog() const
 	const FText TextTitle = FText::FromString(FString::Printf(TEXT("Preflight CL %s: Newer preflight available"), *ServiceProvider.Pin()->GetService<FChangelistService>()->GetCLID()));
 	const FText TextDescription = FText::FromString(FString::Printf(TEXT("There is a newer preflight for this changelist:\n<a id=\"browser\" style=\"Hyperlink\" href=\"%sjob/%s\">%s - %s</>\n\nDo you want to update the #preflight tag?"), *Definition.HordeServerAddress, *HordePreflights->PreflightList[0].ID, *HordePreflights->PreflightList[0].Name, *HordePreflights->PreflightList[0].ID));
 
-	return FDialogFactory::ShowDialog(TextTitle, TextDescription, TArray<FString>{ TEXT("Update Tag Value"), TEXT("Cancel") });
+	TSharedPtr<SHorizontalBox> AutoUpdate = SNew(SHorizontalBox)
+		+ SHorizontalBox::Slot()
+		.AutoWidth()
+		.VAlign(VAlign_Center)
+		[
+			SNew(SCheckBox)
+				.IsChecked_Lambda([]() { return FSubmitToolUserPrefs::Get()->bAutoUpdatePreflight ? ECheckBoxState::Checked : ECheckBoxState::Unchecked; })
+				.OnCheckStateChanged_Lambda([](ECheckBoxState InNewState)
+					{
+						FSubmitToolUserPrefs::Get()->bAutoUpdatePreflight = InNewState == ECheckBoxState::Checked;
+					})
+		]
+		+ SHorizontalBox::Slot()
+		.AutoWidth()
+		.VAlign(VAlign_Center)
+		[
+			SNew(SButton)
+				.ButtonStyle(FAppStyle::Get(), "InvisibleButton")
+				.IsFocusable(false)
+				.OnClicked_Lambda([this]() { FSubmitToolUserPrefs::Get()->bAutoUpdatePreflight = !FSubmitToolUserPrefs::Get()->bAutoUpdatePreflight; return FReply::Handled(); })
+				[
+					SNew(STextBlock)
+						.Justification(ETextJustify::Left)
+						.MinDesiredWidth(60)
+						.Text(FText::FromString(TEXT("Always update, Don't ask again")))
+				]
+		];
+
+	return FDialogFactory::ShowDialog(TextTitle, TextDescription, TArray<FString>{ TEXT("Update Tag Value"), TEXT("Cancel") }, AutoUpdate);
 }
diff --git a/Engine/Source/Programs/SubmitTool/Private/Models/SubmitToolUserPrefs.h b/Engine/Source/Programs/SubmitTool/Private/Models/SubmitToolUserPrefs.h
--- a/Engine/Source/Programs/SubmitTool/Private/Models/SubmitToolUserPrefs.h
+++ b/Engine/Source/Programs/SubmitTool/Private/Models/SubmitToolUserPrefs.h
@@ -27,6 +27,9 @@ struct FSubmitToolUserPrefs
 	UPROPERTY()
 	bool bWindowMaximized = false;
 
+	UPROPERTY()
+	bool bAutoUpdatePreflight = false;
+
 	UPROPERTY()
 	FVector2D WindowPosition = FVector2D();
 
diff --git a/Engine/Source/Programs/SubmitTool/Private/View/Widgets/ConfirmDialogWidget.cpp b/Engine/Source/Programs/SubmitTool/Private/View/Widgets/ConfirmDialogWidget.cpp
--- a/Engine/Source/Programs/SubmitTool/Private/View/Widgets/ConfirmDialogWidget.cpp
+++ b/Engine/Source/Programs/SubmitTool/Private/View/Widgets/ConfirmDialogWidget.cpp
@@ -39,6 +39,12 @@ void SConfirmDialogWidget::Construct(const FArguments& InArgs)
 						}
 					}))
 			]
+			+SVerticalBox::Slot()
+			.AutoHeight()
+			.Padding(5)
+			[
+				InArgs._AdditionalContent == nullptr ? SNullWidget::NullWidget : InArgs._AdditionalContent.ToSharedRef()
+			]
 			+ SVerticalBox::Slot()
 			.AutoHeight()
 			.Padding(0, 5, 0, 0)
diff --git a/Engine/Source/Programs/SubmitTool/Private/View/Widgets/ConfirmDialogWidget.h b/Engine/Source/Programs/SubmitTool/Private/View/Widgets/ConfirmDialogWidget.h
--- a/Engine/Source/Programs/SubmitTool/Private/View/Widgets/ConfirmDialogWidget.h
+++ b/Engine/Source/Programs/SubmitTool/Private/View/Widgets/ConfirmDialogWidget.h
@@ -14,6 +14,7 @@ public:
 	SLATE_BEGIN_ARGS(SConfirmDialogWidget) {}
 		SLATE_ARGUMENT(FText, Title)
 		SLATE_ARGUMENT(FText, DescriptionText)
+		SLATE_ARGUMENT(TSharedPtr<SWidget>, AdditionalContent)
 		SLATE_ARGUMENT(TArray<FString>, Buttons)
 		SLATE_EVENT(FOnResult, ResultCallback)
 	SLATE_END_ARGS()
