Adding skip to delete shelf on automatic submission

diff --git a/Engine/Source/Programs/SubmitTool/Private/Logic/Validators/ValidatorDefinition.h b/Engine/Source/Programs/SubmitTool/Private/Logic/Validators/ValidatorDefinition.h
--- a/Engine/Source/Programs/SubmitTool/Private/Logic/Validators/ValidatorDefinition.h
+++ b/Engine/Source/Programs/SubmitTool/Private/Logic/Validators/ValidatorDefinition.h
@@ -94,15 +94,15 @@ struct FValidatorDefinition
 	FString IncludeFilesInDirectory;
 	
 	///
-	/// Wildcard pattern: only apply this validator if a file matching this pattern is somewhere in the path parent structure
-	/// i.e: *.uproject -> validator will only run if there is a uproject in their parent folder structure
+	/// Wildcard pattern: only apply this validator to files/path which the specified pattern is matched at any level in the parent hierarchy
+	/// i.e: *.uproject -> validator will only include files/paths with a uproject in their parent folder structure
 	/// 
 	UPROPERTY()
 	FString RequireFileInHierarchy;
 
 	///
-	/// Wildcard pattern: Opposite of RequireFileInHierarchy this validator will never run if a file matching this pattern is somewhere in the path parent structure
-	/// i.e: *.uproject -> validator will never run if there is a uproject in their parent folder structure
+	/// Wildcard pattern: Opposite of RequireFileInHierarchy this validator will not apply to files that the pattern is matched at any level in the parent hierarchy
+	/// i.e: *.uproject -> validator will never include files/paths with a uproject in their parent folder structure
 	/// 
 	UPROPERTY()
 	FString ExcludeWhenFileInHierarchy;
diff --git a/Engine/Source/Programs/SubmitTool/Private/Models/ModelInterface.cpp b/Engine/Source/Programs/SubmitTool/Private/Models/ModelInterface.cpp
--- a/Engine/Source/Programs/SubmitTool/Private/Models/ModelInterface.cpp
+++ b/Engine/Source/Programs/SubmitTool/Private/Models/ModelInterface.cpp
@@ -133,7 +133,7 @@ FModelInterface::FModelInterface(const FSubmitToolParameters& InParameters) :
 					if (bAllSucceedIncludingOptional)
 					{
 						UE_LOG(LogSubmitTool, Log, TEXT("User has opted in to automatically submit on validation successful. Proceeding with submission..."));
-						StartSubmitProcess();
+						StartSubmitProcess(true);
 					}
 					else
 					{
@@ -469,7 +469,7 @@ void FModelInterface::RequestSwarmReview(const TArray<FString>& InReviewers)
 	}
 }
 
-void FModelInterface::StartSubmitProcess()
+void FModelInterface::StartSubmitProcess(bool bSkipShelfDialog)
 {
 	PresubmitOperationsService->ResetStates();
 
@@ -569,18 +569,32 @@ void FModelInterface::StartSubmitProcess()
 					ShelveListString = FString::Printf(TEXT("%s\n - And %d other files"), *ShelveListString, ShelvedFiles.Num() - MaxFilesToList);
 				}
 
-				const FString Description =
-					FString::Printf(TEXT("The shelve filelist does not match the local filelist, due to p4 restrictions submit tool can only submit local content do you want to continue with the submit?\nLocal Files:\n - %s\n\nShelved Files:\n - %s"),
-						*LocalListString,
-						*ShelveListString);
-				DialogResult = FDialogFactory::ShowDialog(TextTitle, FText::FromString(Description), TArray<FString>{TEXT("Delete Shelve and Submit"), TEXT("Cancel")});
+				if (bSkipShelfDialog)
+				{
+					DialogResult = EDialogFactoryResult::Confirm;
+				}
+				else
+				{
+					const FString Description =
+						FString::Printf(TEXT("The shelve filelist does not match the local filelist, due to p4 restrictions submit tool can only submit local content do you want to continue with the submit?\nLocal Files:\n - %s\n\nShelved Files:\n - %s"),
+							*LocalListString,
+							*ShelveListString);
+					DialogResult = FDialogFactory::ShowDialog(TextTitle, FText::FromString(Description), TArray<FString>{TEXT("Delete Shelve and Submit"), TEXT("Cancel")});
+				}
 			}
 			else
 			{
-				const FText TextTitle = NSLOCTEXT("SourceControl.SubmitWindow", "DeleteShelvedFilesDialogTitle", "Delete shelved files?");
-				const FText TextDescription = NSLOCTEXT("SourceControl.SubmitWindow", "DeleteShelvedFilesDialogDescription", "There are shelved files in this changelist. Do you want to delete your shelf?\nIf you do not, the submit will be cancelled.");
+				if (bSkipShelfDialog)
+				{
+					DialogResult = EDialogFactoryResult::Confirm;
+				}
+				else
+				{
+					const FText TextTitle = NSLOCTEXT("SourceControl.SubmitWindow", "DeleteShelvedFilesDialogTitle", "Delete shelved files?");
+					const FText TextDescription = NSLOCTEXT("SourceControl.SubmitWindow", "DeleteShelvedFilesDialogDescription", "There are shelved files in this changelist. Do you want to delete your shelf?\nIf you do not, the submit will be cancelled.");
 
-				DialogResult = FDialogFactory::ShowConfirmDialog(TextTitle, TextDescription);
+					DialogResult = FDialogFactory::ShowConfirmDialog(TextTitle, TextDescription);
+				}
 			}
 
 			if (DialogResult == EDialogFactoryResult::Confirm)
diff --git a/Engine/Source/Programs/SubmitTool/Private/Models/ModelInterface.h b/Engine/Source/Programs/SubmitTool/Private/Models/ModelInterface.h
--- a/Engine/Source/Programs/SubmitTool/Private/Models/ModelInterface.h
+++ b/Engine/Source/Programs/SubmitTool/Private/Models/ModelInterface.h
@@ -187,7 +187,7 @@ public:
 
 	void RemovePreflightUpdateCallback(const FDelegateHandle& Handle) const			{ PreflightService->OnPreflightDataUpdated.Remove(Handle); }
 
-	void StartSubmitProcess();
+	void StartSubmitProcess(bool bSkipShelfDialog = false);
 
 	TSharedRef<FSubmitToolServiceProvider> GetServiceProvider() { return ServiceProvider.ToSharedRef(); }
 
