Fix for losing track of shelf state due to comparing pointer references equality rather than actual equality
Fixed infinite loop of deleting shelf when the operation fails in non-standard ways

diff --git a/Engine/Source/Programs/SubmitTool/Private/Logic/ChangelistService.cpp b/Engine/Source/Programs/SubmitTool/Private/Logic/ChangelistService.cpp
--- a/Engine/Source/Programs/SubmitTool/Private/Logic/ChangelistService.cpp
+++ b/Engine/Source/Programs/SubmitTool/Private/Logic/ChangelistService.cpp
@@ -268,7 +268,25 @@ void FChangelistService::RehydrateDataFromP4State()
 		PrintFilesAndShelvedFiles();
 	}
 
-	if(ChangelistState->GetShelvedFilesStates() != ShelvedFilesInCL)
+	const TArray<FSourceControlStateRef> P4Shelf = ChangelistState->GetShelvedFilesStates();
+	bool bNewShelf = false;
+	if (P4Shelf.Num() == ShelvedFilesInCL.Num())
+	{
+		for (size_t i = 0; i < P4Shelf.Num(); ++i)
+		{
+			if (!P4Shelf[i]->GetFilename().Equals(ShelvedFilesInCL[i]->GetFilename(), ESearchCase::IgnoreCase))
+			{
+				bNewShelf = true;
+				break;
+			}
+		}
+	}
+	else
+	{
+		bNewShelf = true;
+	}
+
+	if(bNewShelf)
 	{
 		ChangeType |= ETaskArea::ShelvedFiles;
 		ShelvedFilesInCL = ChangelistState->GetShelvedFilesStates();
diff --git a/Engine/Source/Programs/SubmitTool/Private/Models/ModelInterface.cpp b/Engine/Source/Programs/SubmitTool/Private/Models/ModelInterface.cpp
--- a/Engine/Source/Programs/SubmitTool/Private/Models/ModelInterface.cpp
+++ b/Engine/Source/Programs/SubmitTool/Private/Models/ModelInterface.cpp
@@ -452,7 +452,7 @@ void FModelInterface::RequestSwarmReview(const TArray<FString>& InReviewers)
 		{
 			ChangelistService->CreateShelvedFiles(FSourceControlOperationComplete::CreateLambda([this, InReviewers](const FSourceControlOperationRef& DeleteShelvedOp, ECommandResult::Type Result)
 				{
-					if (Result == ECommandResult::Succeeded)
+					if (Result == ECommandResult::Succeeded && HasShelvedFiles())
 					{
 						RequestSwarmReview(InReviewers);
 					}
