Fixed wrong swarm API usage of reviewer groups when creating a review

diff --git a/Engine/Source/Programs/SubmitTool/Private/Logic/SwarmService.cpp b/Engine/Source/Programs/SubmitTool/Private/Logic/SwarmService.cpp
--- a/Engine/Source/Programs/SubmitTool/Private/Logic/SwarmService.cpp
+++ b/Engine/Source/Programs/SubmitTool/Private/Logic/SwarmService.cpp
@@ -157,7 +157,7 @@ void FSwarmService::CreateReview(const TArray<FString>& InReviewers, const OnCre
 
 	if(!GroupsObject.IsEmpty())
 	{
-		RequestJson->SetArrayField(TEXT("requiredGroups"), GroupsObject);
+		RequestJson->SetArrayField(TEXT("reviewerGroups"), GroupsObject);
 	}
 
 	RequestJson->SetStringField(TEXT("description"), ServiceProvider.Pin()->GetService<FChangelistService>()->GetCLDescription());
