Fix for wrong handling of finished Task Engine events when evaluating synchronous tasks
Version Bump

diff --git a/Engine/Source/Programs/SubmitTool/Private/Logic/TasksService.cpp b/Engine/Source/Programs/SubmitTool/Private/Logic/TasksService.cpp
--- a/Engine/Source/Programs/SubmitTool/Private/Logic/TasksService.cpp
+++ b/Engine/Source/Programs/SubmitTool/Private/Logic/TasksService.cpp
@@ -306,6 +306,7 @@ bool FTasksService::Tick(float InDeltaTime)
 {
 	bool bCurrentCLValidState = true;
 	bool bAnyTaskRunningOrQueued = false;
+	bool bTasksStarted = false;
 	TSet<FName> CurrentGroupsInExecution;
 	
 	for(const TWeakPtr<const FValidatorBase>& Task : CachedTasksWithGroups)
@@ -369,6 +370,7 @@ bool FTasksService::Tick(float InDeltaTime)
 
 				UE_LOG(LogValidatorsResult, Log, TEXT("[%s] Running Task"), *Task->GetValidatorName());
 				Task->StartValidation();
+				bTasksStarted = true;
 			}
 		}
 
@@ -401,11 +403,11 @@ bool FTasksService::Tick(float InDeltaTime)
 		OnTasksRunResultUpdated.Broadcast(bLastTasksRunState);
 	}
 
-	if(bAnyTaskRunningOrQueued != bLastRunningTasks)
+	if(bAnyTaskRunningOrQueued != bLastRunningTasks || bTasksStarted)
 	{
 		bLastRunningTasks = bAnyTaskRunningOrQueued;
 
-		if(!bLastRunningTasks)
+		if(!bAnyTaskRunningOrQueued)
 		{
 			PrintErrorSummary();			
 
diff --git a/Engine/Source/Programs/SubmitTool/Private/Version/AppVersionDefines.h b/Engine/Source/Programs/SubmitTool/Private/Version/AppVersionDefines.h
--- a/Engine/Source/Programs/SubmitTool/Private/Version/AppVersionDefines.h
+++ b/Engine/Source/Programs/SubmitTool/Private/Version/AppVersionDefines.h
@@ -21,7 +21,7 @@
 
 #define SUBMIT_TOOL_VERSION_MAJOR 5
 #define SUBMIT_TOOL_VERSION_MINOR 4
-#define SUBMIT_TOOL_VERSION_PATCH 5
+#define SUBMIT_TOOL_VERSION_PATCH 6
 
 #ifndef SUBMIT_TOOL_CHANGELIST
 #define SUBMIT_TOOL_CHANGELIST_STRING TEXT("NoCL")
