Allow the SubmitTool to open loose package files when running from staged content.

- By default when running a UE process from container files it will not be able to open loose package files. This means that validators that do want to inspect the contents of package files will work when run during development (non staged) but will fail when distributed to end users.
- With this fix SubmitTool will be allowed to open package files from both containers and loose files when staged so that validators behave the same in both development and staged mode.

diff --git a/Engine/Source/Programs/SubmitTool/SubmitTool.Target.cs b/Engine/Source/Programs/SubmitTool/SubmitTool.Target.cs
--- a/Engine/Source/Programs/SubmitTool/SubmitTool.Target.cs
+++ b/Engine/Source/Programs/SubmitTool/SubmitTool.Target.cs
@@ -43,6 +43,7 @@ public class SubmitToolTarget : TargetRules
 		
 		GlobalDefinitions.Add("AUTOSDKS_ENABLED=0");
 		GlobalDefinitions.Add("UE_TRACE_AUTOSTART=1");
+		GlobalDefinitions.Add("EXCLUDE_NONPAK_UE_EXTENSIONS=0");
 
 		OptedInModulePlatforms = new UnrealTargetPlatform[] { UnrealTargetPlatform.Win64, UnrealTargetPlatform.Mac,
 															  UnrealTargetPlatform.Linux, UnrealTargetPlatform.LinuxArm64 };
