Fix for Submit tool not being able to close

diff --git a/Engine/Source/Programs/SubmitTool/Private/View/SubmitToolWindow.cpp b/Engine/Source/Programs/SubmitTool/Private/View/SubmitToolWindow.cpp
--- a/Engine/Source/Programs/SubmitTool/Private/View/SubmitToolWindow.cpp
+++ b/Engine/Source/Programs/SubmitTool/Private/View/SubmitToolWindow.cpp
@@ -20,7 +20,6 @@ TSharedRef<SDockTab> SubmitToolWindow::BuildMainTab(TSharedPtr<SWindow> InParent
 	// Create a submit files window
 	MainTab = SNew(SDockTab)
 		.TabRole(ETabRole::MajorTab)
-		.CanEverClose(false)
 		.OnCanCloseTab(SDockTab::FCanCloseTab::CreateRaw(this, &SubmitToolWindow::OnCanCloseTab))
 		.OnTabClosed_Lambda([](TSharedRef<SDockTab> Tab){
 							if(!Tab->GetParentWindow()->IsWindowMaximized())
