Caching Regex instead of building one everytime

diff --git a/Engine/Source/Programs/SubmitTool/Private/Logic/Validators/ValidatorDefinition.h b/Engine/Source/Programs/SubmitTool/Private/Logic/Validators/ValidatorDefinition.h
--- a/Engine/Source/Programs/SubmitTool/Private/Logic/Validators/ValidatorDefinition.h
+++ b/Engine/Source/Programs/SubmitTool/Private/Logic/Validators/ValidatorDefinition.h
@@ -4,6 +4,7 @@
 
 #include "CoreMinimal.h"
 #include "UObject/Class.h"
+#include "Internationalization/Regex.h"
 
 #include "ValidatorDefinition.generated.h"
 
@@ -270,11 +271,16 @@ struct FValidatorRunExecutableDefinition : public FValidatorDefinition
 	UPROPERTY()
 	FString RegexErrorParsing = TEXT("^(?!.*(?:Display: |Warning: |Log: )).*( error |error:).*$");
 
+	TSharedPtr<FRegexPattern> BuiltRegexError;
+
 	///
 	///	Regex for identifying warnings from the output of a process
 	/// 
 	UPROPERTY()
 	FString RegexWarningParsing = TEXT("^(?!.*(?:Display: |Log: ).*( warning |warning:).*$");
+
+	TSharedPtr<FRegexPattern> BuiltRegexWarning;
+
 };
 
 
diff --git a/Engine/Source/Programs/SubmitTool/Private/Logic/Validators/ValidatorRunExecutable.cpp b/Engine/Source/Programs/SubmitTool/Private/Logic/Validators/ValidatorRunExecutable.cpp
--- a/Engine/Source/Programs/SubmitTool/Private/Logic/Validators/ValidatorRunExecutable.cpp
+++ b/Engine/Source/Programs/SubmitTool/Private/Logic/Validators/ValidatorRunExecutable.cpp
@@ -110,9 +110,11 @@ bool FValidatorRunExecutable::Activate()
 		bIsValidSetup = false;
 	}
 
-	FValidatorRunExecutableDefinition* ModifyableDefinition = const_cast<FValidatorRunExecutableDefinition*>(TypedDefinition);
+	FValidatorRunExecutableDefinition* ModifiableDefinition = const_cast<FValidatorRunExecutableDefinition*>(TypedDefinition);
 
-	ModifyableDefinition->ExecutablePath = FConfiguration::SubstituteAndNormalizeFilename(TypedDefinition->ExecutablePath);
+	ModifiableDefinition->ExecutablePath = FConfiguration::SubstituteAndNormalizeFilename(TypedDefinition->ExecutablePath);
+	ModifiableDefinition->BuiltRegexError = MakeShared<FRegexPattern>(ModifiableDefinition->RegexErrorParsing, ERegexPatternFlags::CaseInsensitive);
+	ModifiableDefinition->BuiltRegexWarning = MakeShared<FRegexPattern>(ModifiableDefinition->RegexWarningParsing, ERegexPatternFlags::CaseInsensitive);
 
 	static TArray<FString> ValidExtensions =
 	{
@@ -259,8 +261,7 @@ bool FValidatorRunExecutable::IsLineAnError(const FString& InLine)
 	}
 
 	{
-		FRegexPattern Pattern = FRegexPattern(TypedDefinition->RegexErrorParsing, ERegexPatternFlags::CaseInsensitive);
-		FRegexMatcher Regex = FRegexMatcher(Pattern, InLine);
+		FRegexMatcher Regex = FRegexMatcher(*TypedDefinition->BuiltRegexError, InLine);
 		if(Regex.FindNext())
 		{
 			return true;
@@ -269,8 +270,7 @@ bool FValidatorRunExecutable::IsLineAnError(const FString& InLine)
 
 	if(Definition->bTreatWarningsAsErrors)
 	{
-		FRegexPattern Pattern = FRegexPattern(TypedDefinition->RegexWarningParsing, ERegexPatternFlags::CaseInsensitive);
-		FRegexMatcher Regex = FRegexMatcher(Pattern, InLine);
+		FRegexMatcher Regex = FRegexMatcher(*TypedDefinition->BuiltRegexWarning, InLine);
 		if(Regex.FindNext())
 		{
 			return true;
