Accumulating OIDC token stdoutput parsing to protect for a possible Json indentation.

Delaying the parsing of the token to process completion

Adding logging to OIDCToken json parse failure

diff --git a/Engine/Source/Programs/SubmitTool/Private/Logic/CredentialsService.cpp b/Engine/Source/Programs/SubmitTool/Private/Logic/CredentialsService.cpp
--- a/Engine/Source/Programs/SubmitTool/Private/Logic/CredentialsService.cpp
+++ b/Engine/Source/Programs/SubmitTool/Private/Logic/CredentialsService.cpp
@@ -150,26 +150,28 @@ void FCredentialsService::GetOIDCToken()
 	{
 		GetOIDCTask = UE::Tasks::Launch(UE_SOURCE_LOCATION,
 			[this] {
-				FOnOutputLine OutputLineProcess = FOnOutputLine::CreateLambda([this](const FString& OutputLine, const EProcessOutputType& OutputType) {
+				FString FullOutput;
+				FOnOutputLine OutputLineProcess = FOnOutputLine::CreateLambda([&FullOutput](const FString& OutputLine, const EProcessOutputType& OutputType) {
 					if(OutputType == EProcessOutputType::ProcessError)
 					{
 						UE_LOG(LogSubmitTool, Error, TEXT("%s"), *OutputLine);
 					}
 					else
 					{
-						// OIDCToken console output token starts with a {
-						if(OutputLine.StartsWith(TEXT("{")))
+						if (OutputType == EProcessOutputType::SDTOutput)
 						{
-							ParseOIDCTokenData(OutputLine);
-						}
-						else
-						{
-							UE_LOG(LogSubmitToolDebug, Log, TEXT("%s"), *OutputLine);
+							FullOutput += OutputLine;
 						}
+
+						UE_LOG(LogSubmitToolDebug, Log, TEXT("%s"), *OutputLine);
 					}
 					});
 
-				OIDCProcess = MakeUnique<FProcessWrapper>(TEXT("Oidc"), Parameters.OAuthTokenTool, FString::Printf(TEXT("%s"), *Parameters.OAuthArgs), FOnCompleted(), OutputLineProcess);
+				FOnCompleted OnCompleted = FOnCompleted::CreateLambda([this, &FullOutput](const int32 InExitCode) {
+					ParseOIDCTokenData(FullOutput);
+					});
+
+				OIDCProcess = MakeUnique<FProcessWrapper>(TEXT("Oidc"), Parameters.OAuthTokenTool, FString::Printf(TEXT("%s"), *Parameters.OAuthArgs), OnCompleted, OutputLineProcess);
 				OIDCProcess->Start(true);
 
 				if(!OIDCProcess.IsValid() || OIDCProcess->ExitCode != 0)
@@ -205,6 +207,8 @@ bool FCredentialsService::ParseOIDCTokenData(const FString& InToken)
 		return true;
 	}
 
+	UE_LOG(LogSubmitTool, Error, TEXT("Couldn't parse OIDC Token from string: '%s'"), *InToken);
+
 	return false;
 }
 
