Warning in preflight validator when submitting with an old preflight (12h since completion)

diff --git a/Engine/Source/Programs/SubmitTool/Private/Logic/Validators/PreflightValidator.cpp b/Engine/Source/Programs/SubmitTool/Private/Logic/Validators/PreflightValidator.cpp
--- a/Engine/Source/Programs/SubmitTool/Private/Logic/Validators/PreflightValidator.cpp
+++ b/Engine/Source/Programs/SubmitTool/Private/Logic/Validators/PreflightValidator.cpp
@@ -6,9 +6,16 @@
 #include "Logic/TagService.h"
 #include "Logging/SubmitToolLog.h"
 #include "Internationalization/Regex.h"
+#include "Models/ModelInterface.h"
 #include "HAL/PlatformFileManager.h"
 #include "GenericPlatform/GenericPlatformFile.h"
 
+FPreflightValidator::FPreflightValidator(const FName& InNameId, const FSubmitToolParameters& InParameters, TSharedRef<FSubmitToolServiceProvider> InServiceProvider, const FString& InDefinition)
+: FValidatorBase(InNameId, InParameters, InServiceProvider, InDefinition)
+{
+	ParseDefinition(InDefinition);
+}
+
 bool FPreflightValidator::Activate()
 {
 	bIsValidSetup = FValidatorBase::Activate();
@@ -19,6 +26,21 @@ bool FPreflightValidator::Activate()
 	return bIsValidSetup;
 }
 
+
+void FPreflightValidator::ParseDefinition(const FString& InDefinition)
+{
+	FStringOutputDevice Errors;
+	Definition = MakeUnique<FPreflightValidatorDefinition>();
+	FPreflightValidatorDefinition* ModifiableDefinition = const_cast<FPreflightValidatorDefinition*>(GetTypedDefinition<FPreflightValidatorDefinition>());
+	FPreflightValidatorDefinition::StaticStruct()->ImportText(*InDefinition, ModifiableDefinition, nullptr, 0, &Errors, FPreflightValidatorDefinition::StaticStruct()->GetName());
+
+	if (!Errors.IsEmpty())
+	{
+		UE_LOG(LogSubmitTool, Error, TEXT("[%s] Error loading parameter file %s"), *GetValidatorName(), *Errors);
+		FModelInterface::SetErrorState();
+	}
+}
+
 bool FPreflightValidator::Validate(const FString& InCLDescription, const TArray<FSourceControlStateRef>& InFilteredFilesInCL, const TArray<const FTag*>& InTags)
 {
 	if(!PreflightTag->IsEnabled() || PreflightTag->GetValues().Num() == 0)
@@ -78,6 +100,7 @@ void FPreflightValidator::ValidatePreflights(const TUniquePtr<FPreflightList>& I
 	size_t RunningPreflights = 0;
 	TArray<const FPreflightData*, TInlineAllocator<8>> ValidPreflights;
 	TArray<const FPreflightData*, TInlineAllocator<8>> FailedPreflights;
+	TArray<const FPreflightData*, TInlineAllocator<8>> WarningPreflights;
 	TArray<FString, TInlineAllocator<8>> UnverifiedPreflights;
 
 	const TArray<FString>& PreflightValues = PreflightTag->GetValues();
@@ -130,6 +153,7 @@ void FPreflightValidator::ValidatePreflights(const TUniquePtr<FPreflightList>& I
 						}
 					}
 
+					FTimespan PreflightAge = FDateTime::UtcNow() - FoundData->UpdateTime;
 					switch(FoundData->CachedResults.State)
 					{
 						case EPreflightState::Running:
@@ -138,26 +162,34 @@ void FPreflightValidator::ValidatePreflights(const TUniquePtr<FPreflightList>& I
 							break;
 
 						case EPreflightState::Completed:
-							switch(FoundData->CachedResults.Outcome)
+							if (PreflightAge.GetTotalHours() > GetTypedDefinition<FPreflightValidatorDefinition>()->MaxPreflightAgeInHours)
+							{
+								LogFailure(FString::Printf(TEXT("[%s] Preflight %s is %d hours old, submitting changes with preflights older than %d hours is not recommended."), *GetValidatorName(), *PreflightId, FMath::FloorToInt(PreflightAge.GetTotalHours()), GetTypedDefinition<FPreflightValidatorDefinition>()->MaxPreflightAgeInHours));
+								WarningPreflights.Add(FoundData);
+							}
+							else
 							{
-								case EPreflightOutcome::Success:
-									FoundData->CachedResults.WasSuccessful() ? ValidPreflights.Add(FoundData) : FailedPreflights.Add(FoundData);
-									break;
-								case EPreflightOutcome::Unspecified:
-								case EPreflightOutcome::Failure:
-									FailedPreflights.Add(FoundData);
-									break;
-								case EPreflightOutcome::Warnings:
-									if(Definition->bTreatWarningsAsErrors)
-									{
-										LogFailure(FString::Printf(TEXT("[%s] %s preflight has completed with warnings and they are treated as errors."), *GetValidatorName(), *PreflightId));
+								switch(FoundData->CachedResults.Outcome)
+								{
+									case EPreflightOutcome::Success:
+										FoundData->CachedResults.WasSuccessful() ? ValidPreflights.Add(FoundData) : FailedPreflights.Add(FoundData);
+										break;
+									case EPreflightOutcome::Unspecified:
+									case EPreflightOutcome::Failure:
 										FailedPreflights.Add(FoundData);
-									}
-									else
-									{
-										ValidPreflights.Add(FoundData);
-									}
-									break;
+										break;
+									case EPreflightOutcome::Warnings:
+										if(Definition->bTreatWarningsAsErrors)
+										{
+											LogFailure(FString::Printf(TEXT("[%s] %s preflight has completed with warnings and they are treated as errors."), *GetValidatorName(), *PreflightId));
+											FailedPreflights.Add(FoundData);
+										}
+										else
+										{
+											ValidPreflights.Add(FoundData);
+										}
+										break;
+								}
 							}
 							break;
 
@@ -209,6 +241,11 @@ void FPreflightValidator::ValidatePreflights(const TUniquePtr<FPreflightList>& I
 	}
 	else if(RunningPreflights == 0)
 	{
+		for (const FPreflightData* PreflightPtr : WarningPreflights)
+		{
+			UE_LOG(LogValidators, Warning, TEXT("[%s] %s preflight has warnings."), *GetValidatorName(), *PreflightPtr->ID);
+		}
+
 		for(const FString& PreflightId : UnverifiedPreflights)
 		{
 			UE_LOG(LogValidators, Error, TEXT("[%s] %s preflight can't be verified with Horde, check that the id is correct and there are no connection errors."), *GetValidatorName(), *PreflightId);
diff --git a/Engine/Source/Programs/SubmitTool/Private/Logic/Validators/PreflightValidator.h b/Engine/Source/Programs/SubmitTool/Private/Logic/Validators/PreflightValidator.h
--- a/Engine/Source/Programs/SubmitTool/Private/Logic/Validators/PreflightValidator.h
+++ b/Engine/Source/Programs/SubmitTool/Private/Logic/Validators/PreflightValidator.h
@@ -10,8 +10,9 @@ class FPreflightService;
 class FPreflightValidator : public FValidatorBase
 {
 public:
-	using FValidatorBase::FValidatorBase;
+	FPreflightValidator(const FName& InNameId, const FSubmitToolParameters& InParameters, TSharedRef<FSubmitToolServiceProvider> InServiceProvider, const FString& InDefinition);
 
+	virtual void ParseDefinition(const FString& InDefinition) override;
 	virtual bool Activate() override;
 
 	virtual bool Validate(const FString& InCLDescription, const TArray<FSourceControlStateRef>& InFilteredFilesInCL, const TArray<const FTag*>& InTags) override;
diff --git a/Engine/Source/Programs/SubmitTool/Private/Logic/Validators/ValidatorDefinition.h b/Engine/Source/Programs/SubmitTool/Private/Logic/Validators/ValidatorDefinition.h
--- a/Engine/Source/Programs/SubmitTool/Private/Logic/Validators/ValidatorDefinition.h
+++ b/Engine/Source/Programs/SubmitTool/Private/Logic/Validators/ValidatorDefinition.h
@@ -347,6 +347,20 @@ struct FJsonValidatorDefinition : public FValidatorDefinition
 	FString RegexLineExclusion;
 };
 
+
+USTRUCT()
+struct FPreflightValidatorDefinition : public FValidatorDefinition
+{
+	GENERATED_BODY()
+
+	///
+	/// The maximum hours since completion that a preflight can be accepted as a success
+	/// 
+	UPROPERTY()
+	int32 MaxPreflightAgeInHours = 12;
+};
+
+
 USTRUCT()
 struct FVirtualizationToolDefinition : public FValidatorRunExecutableDefinition
 {
