Adding tooltips to tag labels

diff --git a/Engine/Source/Programs/SubmitTool/Private/View/Widgets/TagWidget.cpp b/Engine/Source/Programs/SubmitTool/Private/View/Widgets/TagWidget.cpp
--- a/Engine/Source/Programs/SubmitTool/Private/View/Widgets/TagWidget.cpp
+++ b/Engine/Source/Programs/SubmitTool/Private/View/Widgets/TagWidget.cpp
@@ -115,6 +115,7 @@ void STagWidget::Construct(const FArguments& InArgs)
 					return FLinearColor::White;
 			}})
 			.Text_Lambda([this](){ return FText::FromString(FString::Printf(TEXT("%s%s"), *this->Tag->Definition.TagLabel, Tag->GetCurrentValidationConfig(ModelInterface->GetDepotFilesInCL()).bIsMandatory ? TEXT(" *") : TEXT(""))); })
+			.ToolTipText(FText::FromString(Tag->Definition.ToolTip))
 			.Font(FAppStyle::Get().GetFontStyle("SmallFont"))
 		]
 	];
