Adding option to auto prepend users and groups with [at]

diff --git a/Engine/Source/Programs/SubmitTool/Private/Models/SubmitToolUserPrefs.h b/Engine/Source/Programs/SubmitTool/Private/Models/SubmitToolUserPrefs.h
--- a/Engine/Source/Programs/SubmitTool/Private/Models/SubmitToolUserPrefs.h
+++ b/Engine/Source/Programs/SubmitTool/Private/Models/SubmitToolUserPrefs.h
@@ -54,6 +54,12 @@ struct FSubmitToolUserPrefs
 
 	UPROPERTY()
 	bool bAutoUpdate = true;
+	
+	UPROPERTY()
+	bool bAppendAtForP4Users = true;
+
+	UPROPERTY()
+	bool bAppendAtForP4Groups = true;
 
 	virtual ~FSubmitToolUserPrefs();
 
diff --git a/Engine/Source/Programs/SubmitTool/Private/View/Widgets/SelectSourceControlGroupWidget.cpp b/Engine/Source/Programs/SubmitTool/Private/View/Widgets/SelectSourceControlGroupWidget.cpp
--- a/Engine/Source/Programs/SubmitTool/Private/View/Widgets/SelectSourceControlGroupWidget.cpp
+++ b/Engine/Source/Programs/SubmitTool/Private/View/Widgets/SelectSourceControlGroupWidget.cpp
@@ -2,6 +2,7 @@
 
 #include "SelectSourceControlGroupWidget.h"
 
+#include "Models/SubmitToolUserPrefs.h"
 #include "View/SubmitToolStyle.h"
 #include "Widgets/SWindow.h"
 #include "Widgets/Images/SThrobber.h"
@@ -112,6 +113,32 @@ FReply SSelectSourceControlGroupWidget::OpenDialog()
 			SAssignNew(FilterTextBox, SEditableTextBox)
 				.Text_Lambda([&text = FilterText] {return FText::FromString(text); })
 				.OnTextChanged_Lambda([this](FText text) { this->OnFilterTextChanged(text.ToString()); })
+		]
+		+SHorizontalBox::Slot()
+		.AutoWidth()
+		.Padding(5, 0, 0, 0)
+		.VAlign(VAlign_Center)
+		[
+			SNew(SCheckBox)
+			.IsChecked_Lambda([]() { return FSubmitToolUserPrefs::Get()->bAppendAtForP4Groups ? ECheckBoxState::Checked : ECheckBoxState::Unchecked; })
+			.OnCheckStateChanged_Lambda([](ECheckBoxState InNewState) { FSubmitToolUserPrefs::Get()->bAppendAtForP4Groups = InNewState == ECheckBoxState::Checked; })
+		]
+		+SHorizontalBox::Slot()
+		.AutoWidth()
+		.VAlign(VAlign_Center)
+		[
+			SNew(SButton)
+			.ButtonStyle(FAppStyle::Get(), "InvisibleButton")
+			.IsFocusable(false)			
+
+			.OnClicked_Lambda([]() { FSubmitToolUserPrefs::Get()->bAppendAtForP4Groups = !FSubmitToolUserPrefs::Get()->bAppendAtForP4Groups; return FReply::Handled(); })
+			[
+				SNew(STextBlock)
+				.Justification(ETextJustify::Left)
+				.MinDesiredWidth(60)
+				.Text(FText::FromString(TEXT("Append @")))
+				.ToolTipText(FText::FromString(TEXT("Appends an @ at the beginning of the group name so that p4 notifies when requesting a review")))
+			]
 		];
 
 	Contents.Get()->AddSlot()
@@ -325,9 +352,9 @@ FReply SSelectSourceControlGroupWidget::OnCancelClicked()
 
 void SSelectSourceControlGroupWidget::ProcessGroupSelected(TSharedPtr<FString> InGroup)
 {
-	const FString& GroupName = *InGroup;
+	const FString GroupName = FSubmitToolUserPrefs::Get()->bAppendAtForP4Groups ? TEXT("@") + *InGroup : *InGroup;
 	TArray<FString> CurrentValues = TargetTag->GetValues();
-	if (!CurrentValues.ContainsByPredicate([GroupName](const FString& Value) { return Value.Equals(GroupName, ESearchCase::IgnoreCase); }))
+	if (!CurrentValues.ContainsByPredicate([GroupName](const FString& Value) { return Value.TrimChar('@').Equals(GroupName.TrimChar('@'), ESearchCase::IgnoreCase); }))
 	{
 		CurrentValues.Add(GroupName);
 	}
diff --git a/Engine/Source/Programs/SubmitTool/Private/View/Widgets/SelectSourceControlUserWidget.cpp b/Engine/Source/Programs/SubmitTool/Private/View/Widgets/SelectSourceControlUserWidget.cpp
--- a/Engine/Source/Programs/SubmitTool/Private/View/Widgets/SelectSourceControlUserWidget.cpp
+++ b/Engine/Source/Programs/SubmitTool/Private/View/Widgets/SelectSourceControlUserWidget.cpp
@@ -2,6 +2,7 @@
 
 #include "SelectSourceControlUserWidget.h"
 
+#include "Models/SubmitToolUserPrefs.h"
 #include "View/SubmitToolStyle.h"
 #include "Widgets/SWindow.h"
 #include "Widgets/Images/SThrobber.h"
@@ -117,6 +118,31 @@ FReply SSelectSourceControlUserWidget::OpenDialog()
 			SAssignNew(FilterTextBox, SEditableTextBox)
 			.Text_Lambda([&text = FilterText] {return FText::FromString(text); })
 			.OnTextChanged_Lambda([this](FText text) { this->OnFilterTextChanged(text.ToString()); })
+		]
+		+SHorizontalBox::Slot()
+		.AutoWidth()
+		.VAlign(VAlign_Center)
+		[
+			SNew(SCheckBox)
+			.IsChecked_Lambda([]() { return FSubmitToolUserPrefs::Get()->bAppendAtForP4Users ? ECheckBoxState::Checked : ECheckBoxState::Unchecked; })
+			.OnCheckStateChanged_Lambda([](ECheckBoxState InNewState) { FSubmitToolUserPrefs::Get()->bAppendAtForP4Users = InNewState == ECheckBoxState::Checked; })
+		]
+		+SHorizontalBox::Slot()
+		.AutoWidth()
+		.VAlign(VAlign_Center)
+		[
+			SNew(SButton)
+			.ButtonStyle(FAppStyle::Get(), "InvisibleButton")
+			.IsFocusable(false)			
+
+			.OnClicked_Lambda([]() { FSubmitToolUserPrefs::Get()->bAppendAtForP4Users = !FSubmitToolUserPrefs::Get()->bAppendAtForP4Users; return FReply::Handled(); })
+			[
+				SNew(STextBlock)
+				.Justification(ETextJustify::Left)
+				.MinDesiredWidth(60)
+				.Text(FText::FromString(TEXT("Append @")))
+				.ToolTipText(FText::FromString(TEXT("Appends an @ at the beginning of the user name so that p4 notifies when requesting a review")))
+			]
 		];
 
 	if(TargetTag && TargetTag->Definition.Filters.Num() != 0)
@@ -405,7 +431,7 @@ FReply SSelectSourceControlUserWidget::OnCancelClicked()
 
 void SSelectSourceControlUserWidget::ProcessUserSelected(TSharedPtr<FUserData> InUser)
 {
-	const FString& Username = InUser->Username;
+	const FString Username = FSubmitToolUserPrefs::Get()->bAppendAtForP4Users ? TEXT("@") + InUser->Username : InUser->Username;
 
 	TArray<FString> CurrentValues;
 	if(TargetTag)
@@ -423,7 +449,7 @@ void SSelectSourceControlUserWidget::ProcessUserSelected(TSharedPtr<FUserData> I
 		TargetText->ParseIntoArray(CurrentValues, Delimiters.GetData(), 2);
 	}
 
-	if(!CurrentValues.ContainsByPredicate([Username](const FString& Value) { return Value.Equals(Username, ESearchCase::IgnoreCase); }))
+	if(!CurrentValues.ContainsByPredicate([Username](const FString& Value) { return Value.TrimChar('@').Equals(Username.TrimChar('@'), ESearchCase::IgnoreCase); }))
 	{
 		CurrentValues.Add(Username);
 		ModelInterface->AddRecentUser(InUser);
