Fixed issues with removing usernames with [at] when value didn't have the [at]
Removed the [at] option when targetting SwarmApproved tags since they're not compatible with UGS

diff --git a/Engine/Source/Programs/SubmitTool/Private/View/Widgets/SelectSourceControlGroupWidget.cpp b/Engine/Source/Programs/SubmitTool/Private/View/Widgets/SelectSourceControlGroupWidget.cpp
--- a/Engine/Source/Programs/SubmitTool/Private/View/Widgets/SelectSourceControlGroupWidget.cpp
+++ b/Engine/Source/Programs/SubmitTool/Private/View/Widgets/SelectSourceControlGroupWidget.cpp
@@ -238,7 +238,7 @@ FReply SSelectSourceControlGroupWidget::OpenDialog()
 
 						FString GroupName = *selectedGroup;
 						TArray<FString> CurrentValues = targetTag->GetValues();
-						if (!CurrentValues.ContainsByPredicate([GroupName](const FString& Value) { return Value.Equals(GroupName, ESearchCase::IgnoreCase); }))
+						if (!CurrentValues.ContainsByPredicate([GroupName](const FString& Value) { return Value.TrimChar('@').Equals(GroupName.TrimChar('@'), ESearchCase::IgnoreCase); }))
 						{
 							return FText::FromString("Add group");
 						}
@@ -360,7 +360,7 @@ void SSelectSourceControlGroupWidget::ProcessGroupSelected(TSharedPtr<FString> I
 	}
 	else
 	{
-		CurrentValues.Remove(GroupName);
+		CurrentValues.RemoveAll([GroupName](const FString& Value) { return Value.TrimChar('@').Equals(GroupName.TrimChar('@'), ESearchCase::IgnoreCase); });
 	}
 
 	ModelInterface->SetTagValues(*TargetTag, CurrentValues);
diff --git a/Engine/Source/Programs/SubmitTool/Private/View/Widgets/SelectSourceControlUserWidget.cpp b/Engine/Source/Programs/SubmitTool/Private/View/Widgets/SelectSourceControlUserWidget.cpp
--- a/Engine/Source/Programs/SubmitTool/Private/View/Widgets/SelectSourceControlUserWidget.cpp
+++ b/Engine/Source/Programs/SubmitTool/Private/View/Widgets/SelectSourceControlUserWidget.cpp
@@ -121,9 +121,11 @@ FReply SSelectSourceControlUserWidget::OpenDialog()
 		]
 		+SHorizontalBox::Slot()
 		.AutoWidth()
+		.Padding(5, 0, 0, 0)
 		.VAlign(VAlign_Center)
 		[
 			SNew(SCheckBox)
+			.Visibility_Lambda([this](){ return TargetTag->Definition.InputSubType == TEXT("SwarmApproved") ? EVisibility::Collapsed : EVisibility::All ;})
 			.IsChecked_Lambda([]() { return FSubmitToolUserPrefs::Get()->bAppendAtForP4Users ? ECheckBoxState::Checked : ECheckBoxState::Unchecked; })
 			.OnCheckStateChanged_Lambda([](ECheckBoxState InNewState) { FSubmitToolUserPrefs::Get()->bAppendAtForP4Users = InNewState == ECheckBoxState::Checked; })
 		]
@@ -133,8 +135,8 @@ FReply SSelectSourceControlUserWidget::OpenDialog()
 		[
 			SNew(SButton)
 			.ButtonStyle(FAppStyle::Get(), "InvisibleButton")
-			.IsFocusable(false)			
-
+			.Visibility_Lambda([this]() { return TargetTag->Definition.InputSubType == TEXT("SwarmApproved") ? EVisibility::Collapsed : EVisibility::All; })
+			.IsFocusable(false)		
 			.OnClicked_Lambda([]() { FSubmitToolUserPrefs::Get()->bAppendAtForP4Users = !FSubmitToolUserPrefs::Get()->bAppendAtForP4Users; return FReply::Handled(); })
 			[
 				SNew(STextBlock)
@@ -294,7 +296,7 @@ FReply SSelectSourceControlUserWidget::OpenDialog()
 				if(TargetTag)
 				{
 					TArray<FString> CurrentValues = TargetTag->GetValues();
-					bExists = CurrentValues.ContainsByPredicate([UserName](const FString& Value) { return Value.Equals(UserName, ESearchCase::IgnoreCase); });
+					bExists = CurrentValues.ContainsByPredicate([UserName](const FString& Value) { return Value.TrimChar('@').Equals(UserName.TrimChar('@'), ESearchCase::IgnoreCase); });
 				}
 				else
 				{
@@ -431,7 +433,7 @@ FReply SSelectSourceControlUserWidget::OnCancelClicked()
 
 void SSelectSourceControlUserWidget::ProcessUserSelected(TSharedPtr<FUserData> InUser)
 {
-	const FString Username = FSubmitToolUserPrefs::Get()->bAppendAtForP4Users ? TEXT("@") + InUser->Username : InUser->Username;
+	const FString Username = TargetTag->Definition.InputSubType != TEXT("SwarmApproved") && FSubmitToolUserPrefs::Get()->bAppendAtForP4Users ? TEXT("@") + InUser->Username : InUser->Username;
 
 	TArray<FString> CurrentValues;
 	if(TargetTag)
@@ -456,7 +458,7 @@ void SSelectSourceControlUserWidget::ProcessUserSelected(TSharedPtr<FUserData> I
 	}
 	else
 	{
-		CurrentValues.Remove(Username);
+		CurrentValues.RemoveAll([Username](const FString& Value) { return Value.TrimChar('@').Equals(Username.TrimChar('@'), ESearchCase::IgnoreCase); });
 	}
 
 	if(TargetTag)
