Fix for virtualization operation not initializing the regex and causing a nullptr crash

diff --git a/Engine/Source/Programs/SubmitTool/Private/Logic/PreSubmitOperations/VirtualizationOperation.cpp b/Engine/Source/Programs/SubmitTool/Private/Logic/PreSubmitOperations/VirtualizationOperation.cpp
--- a/Engine/Source/Programs/SubmitTool/Private/Logic/PreSubmitOperations/VirtualizationOperation.cpp
+++ b/Engine/Source/Programs/SubmitTool/Private/Logic/PreSubmitOperations/VirtualizationOperation.cpp
@@ -58,7 +58,7 @@ void FVirtualizationOperation::ParseDefinition(const FString& InDefinition)
 
 bool FVirtualizationOperation::Activate()
 {
-	bIsValidSetup = FValidatorBase::Activate();
+	bIsValidSetup = FValidatorRunExecutable::Activate();
 	const FValidatorRunExecutableDefinition* TypedDefinition = GetTypedDefinition<FValidatorRunExecutableDefinition>();
 	static TArray<FString> ValidExtensions =
 	{
