Added bValidateExecutableExists to allow tasks that compile themselves as part of their execution like VirtualizationTool

diff --git a/Engine/Source/Programs/SubmitTool/Private/Logic/PreSubmitOperations/VirtualizationOperation.cpp b/Engine/Source/Programs/SubmitTool/Private/Logic/PreSubmitOperations/VirtualizationOperation.cpp
--- a/Engine/Source/Programs/SubmitTool/Private/Logic/PreSubmitOperations/VirtualizationOperation.cpp
+++ b/Engine/Source/Programs/SubmitTool/Private/Logic/PreSubmitOperations/VirtualizationOperation.cpp
@@ -58,31 +58,9 @@ void FVirtualizationOperation::ParseDefinition(const FString& InDefinition)
 
 bool FVirtualizationOperation::Activate()
 {
+	FVirtualizationToolDefinition* ModifyableDefinition = const_cast<FVirtualizationToolDefinition*>(GetTypedDefinition<FVirtualizationToolDefinition>());
+	ModifyableDefinition->bValidateExecutableExists = false;
 	bIsValidSetup = FValidatorRunExecutable::Activate();
-	const FValidatorRunExecutableDefinition* TypedDefinition = GetTypedDefinition<FValidatorRunExecutableDefinition>();
-	static TArray<FString> ValidExtensions =
-	{
-#if PLATFORM_WINDOWS
-		TEXT(".exe"),
-		TEXT(".bat"),
-#elif PLATFORM_MAC
-		TEXT(".app"),
-		TEXT(".sh"),
-		TEXT(".command"),
-		TEXT(""),
-#elif PLATFORM_LINUX
-		TEXT(".sh"),
-		TEXT(""),
-#endif
-	};
-
-	bIsValidSetup = bIsValidSetup && !TypedDefinition->ExecutablePath.IsEmpty();
-
-	if(!ValidExtensions.Contains(FPaths::GetExtension(TypedDefinition->ExecutablePath, true)))
-	{
-		UE_LOG(LogSubmitTool, Error, TEXT("Task '%s' executable has an invalid extension for this platform: %s"), *ValidatorName, *TypedDefinition->ExecutablePath);
-		bIsValidSetup = false;
-	}
 
 	return bIsValidSetup;
 }
diff --git a/Engine/Source/Programs/SubmitTool/Private/Logic/Validators/ValidatorDefinition.h b/Engine/Source/Programs/SubmitTool/Private/Logic/Validators/ValidatorDefinition.h
--- a/Engine/Source/Programs/SubmitTool/Private/Logic/Validators/ValidatorDefinition.h
+++ b/Engine/Source/Programs/SubmitTool/Private/Logic/Validators/ValidatorDefinition.h
@@ -222,6 +222,9 @@ struct FValidatorRunExecutableDefinition : public FValidatorDefinition
 	UPROPERTY()
 	bool bLaunchReallyHidden = true;
 
+	UPROPERTY()
+	bool bValidateExecutableExists = true;
+
 	///
 	///	Path to the executable that this validator runs
 	/// 
diff --git a/Engine/Source/Programs/SubmitTool/Private/Logic/Validators/ValidatorRunExecutable.cpp b/Engine/Source/Programs/SubmitTool/Private/Logic/Validators/ValidatorRunExecutable.cpp
--- a/Engine/Source/Programs/SubmitTool/Private/Logic/Validators/ValidatorRunExecutable.cpp
+++ b/Engine/Source/Programs/SubmitTool/Private/Logic/Validators/ValidatorRunExecutable.cpp
@@ -157,7 +157,7 @@ bool FValidatorRunExecutable::Activate()
 			bIsValidSetup = false;
 		}
 
-		if(!FPaths::FileExists(TypedDefinition->ExecutablePath))
+		if(TypedDefinition->bValidateExecutableExists && !FPaths::FileExists(TypedDefinition->ExecutablePath))
 		{
 			ActivationErrors.Add(FString::Printf(TEXT("Task '%s' executable is not found on disk: %s."), *ValidatorName, *TypedDefinition->ExecutablePath));
 			bIsValidSetup = false;
@@ -261,6 +261,7 @@ bool FValidatorRunExecutable::IsLineAnError(const FString& InLine)
 	}
 
 	{
+
 		FRegexMatcher Regex = FRegexMatcher(*TypedDefinition->BuiltRegexError, InLine);
 		if(Regex.FindNext())
 		{
diff --git a/Engine/Source/Programs/SubmitTool/Private/Version/AppVersionDefines.h b/Engine/Source/Programs/SubmitTool/Private/Version/AppVersionDefines.h
--- a/Engine/Source/Programs/SubmitTool/Private/Version/AppVersionDefines.h
+++ b/Engine/Source/Programs/SubmitTool/Private/Version/AppVersionDefines.h
@@ -21,7 +21,7 @@
 
 #define SUBMIT_TOOL_VERSION_MAJOR 5
 #define SUBMIT_TOOL_VERSION_MINOR 5
-#define SUBMIT_TOOL_VERSION_PATCH 1
+#define SUBMIT_TOOL_VERSION_PATCH 2
 
 #ifndef SUBMIT_TOOL_CHANGELIST
 #define SUBMIT_TOOL_CHANGELIST_STRING TEXT("NoCL")
