Fixed crashes when callbacks from perforce groups come from a different thread

diff --git a/Engine/Source/Programs/SubmitTool/Private/Logic/Services/SourceControl/SubmitToolPerforce.cpp b/Engine/Source/Programs/SubmitTool/Private/Logic/Services/SourceControl/SubmitToolPerforce.cpp
--- a/Engine/Source/Programs/SubmitTool/Private/Logic/Services/SourceControl/SubmitToolPerforce.cpp
+++ b/Engine/Source/Programs/SubmitTool/Private/Logic/Services/SourceControl/SubmitToolPerforce.cpp
@@ -8,6 +8,7 @@
 #include "Configuration/Configuration.h"
 #include "Logging/SubmitToolLog.h"
 #include "Logic/CredentialsService.h"
+#include "Async/Async.h"
 
 #include "CommandLine/CmdLineParameters.h"
 #include "Misc/Base64.h"
@@ -560,9 +561,12 @@ void FSubmitToolPerforce::GetGroups(const FOnGroupsGet::FDelegate& Callback)
 
 				LoadRecentGroups();
 
-				FScopeLock Lock(&Mutex);
-				OnGroupsGetCallbacks.Broadcast(CachedGroupsArray);
-				OnGroupsGetCallbacks.Clear();
+				AsyncTask(ENamedThreads::GameThread, [this] 
+				{
+					FScopeLock Lock(&Mutex);
+					OnGroupsGetCallbacks.Broadcast(CachedGroupsArray);
+					OnGroupsGetCallbacks.Clear(); 
+				});
 
 				return !ResultInfo.HasErrors();
 
