Fix compilation issues for Mac & Linux
Submit Tool version bump

diff --git a/Engine/Source/Programs/SubmitTool/Private/Logic/Validators/ValidatorRunExecutable.cpp b/Engine/Source/Programs/SubmitTool/Private/Logic/Validators/ValidatorRunExecutable.cpp
--- a/Engine/Source/Programs/SubmitTool/Private/Logic/Validators/ValidatorRunExecutable.cpp
+++ b/Engine/Source/Programs/SubmitTool/Private/Logic/Validators/ValidatorRunExecutable.cpp
@@ -319,8 +319,11 @@ void FValidatorRunExecutable::Tick(float DeltaTime)
 			bAnyProcessRunning |= Process.IsRunning();
 			bAllFinished &= Process.bIsComplete;
 
-			const bool Success = Def->AllowedExitCodes.Contains(Process.ExitCode) && (Def->bOnlyLookAtExitCode || ProcessesData[Process.GetProcessName()].ErrorList.IsEmpty());
-			bAllSucceeded &= Success;
+			if (Process.bIsComplete)
+			{
+				const bool Success = Def->AllowedExitCodes.Contains(Process.ExitCode) && (Def->bOnlyLookAtExitCode || ProcessesData[Process.GetProcessName()].ErrorList.IsEmpty());
+				bAllSucceeded &= Success;
+			}
 		}
 
 		if (bAllFinished)
@@ -383,12 +386,12 @@ const TArray<FAnalyticsEventAttribute> FValidatorRunExecutable::GetTelemetryAttr
 		else
 		{
 			AppendAnalyticsEventAttributeArray(Attributes,
-				TEXT("Executable #") + i , ProcessWrapper.GetExecutable(),
-				TEXT("Args #") + i, ProcessWrapper.GetArgs(),
-				TEXT("Started #") + i, ProcessWrapper.bStarted,
-				TEXT("ExeExitCode #") + i, ProcessWrapper.ExitCode,
-				TEXT("ExeRunTime #") + i, ProcessWrapper.ExecutingTime,
-				TEXT("ErrorCount #") + i, ProcessesData[ProcessWrapper.GetProcessName()].ErrorList.Num()
+				FString::Printf(TEXT("Executable #%i"), i) , ProcessWrapper.GetExecutable(),
+				FString::Printf(TEXT("Args #%i"), i), ProcessWrapper.GetArgs(),
+				FString::Printf(TEXT("Started #%i"), i), ProcessWrapper.bStarted,
+				FString::Printf(TEXT("ExeExitCode #%i"), i), ProcessWrapper.ExitCode,
+				FString::Printf(TEXT("ExeRunTime #%i"), i), ProcessWrapper.ExecutingTime,
+				FString::Printf(TEXT("ErrorCount #%i"), i), ProcessesData[ProcessWrapper.GetProcessName()].ErrorList.Num()
 			);
 		}
 
diff --git a/Engine/Source/Programs/SubmitTool/Private/Version/AppVersionDefines.h b/Engine/Source/Programs/SubmitTool/Private/Version/AppVersionDefines.h
--- a/Engine/Source/Programs/SubmitTool/Private/Version/AppVersionDefines.h
+++ b/Engine/Source/Programs/SubmitTool/Private/Version/AppVersionDefines.h
@@ -19,9 +19,9 @@
 #define SUBMIT_TOOL_APPNAME TEXT("SubmitTool")
 #endif
 
-#define SUBMIT_TOOL_VERSION_MAJOR 5
-#define SUBMIT_TOOL_VERSION_MINOR 5
-#define SUBMIT_TOOL_VERSION_PATCH 3
+#define SUBMIT_TOOL_VERSION_MAJOR 6
+#define SUBMIT_TOOL_VERSION_MINOR 0
+#define SUBMIT_TOOL_VERSION_PATCH 0
 
 #ifndef SUBMIT_TOOL_CHANGELIST
 #define SUBMIT_TOOL_CHANGELIST_STRING TEXT("NoCL")
