Fix submit tool update mechanism being killed by FProcessWrapper

diff --git a/Engine/Source/Programs/SubmitTool/Private/Logic/ProcessWrapper.cpp b/Engine/Source/Programs/SubmitTool/Private/Logic/ProcessWrapper.cpp
--- a/Engine/Source/Programs/SubmitTool/Private/Logic/ProcessWrapper.cpp
+++ b/Engine/Source/Programs/SubmitTool/Private/Logic/ProcessWrapper.cpp
@@ -8,7 +8,7 @@
 #include "Misc/DateTime.h"
 #include "ProcessPipes.h"
 
-FProcessWrapper::FProcessWrapper(const FString& InProcessName, const FString& InPath, const FString& InArgs, const FOnCompleted& InOnCompleted, const FOnOutputLine& InOnOutputLine, const FString& InWorkingDir, const bool InbLaunchHidden, const bool bLaunchReallyHidden, const bool InbLaunchDetached)
+FProcessWrapper::FProcessWrapper(const FString& InProcessName, const FString& InPath, const FString& InArgs, const FOnCompleted& InOnCompleted, const FOnOutputLine& InOnOutputLine, const FString& InWorkingDir, const bool InbLaunchHidden, const bool bLaunchReallyHidden, const bool InbLaunchDetached, bool InbOwnsProcessLifetime)
 	: ProcessName(InProcessName)
 	, Path(InPath)
 	, Args(InArgs)
@@ -16,6 +16,7 @@ FProcessWrapper::FProcessWrapper(const FString& InProcessName, const FString& In
 	, bLaunchesHidden(InbLaunchHidden)
 	, bLaunchesReallyHidden(bLaunchReallyHidden)
 	, bLaunchDetached(InbLaunchDetached)
+	, bOwnsProcessLifetime(InbOwnsProcessLifetime)
 	, Pipes(MakeUnique<FProcessPipes>())
 	, OnCompleted(InOnCompleted)
 	, OnOutputLine(InOnOutputLine)
@@ -23,7 +24,14 @@ FProcessWrapper::FProcessWrapper(const FString& InProcessName, const FString& In
 
 FProcessWrapper::~FProcessWrapper()
 {
-	Stop();
+	if (bOwnsProcessLifetime)
+	{
+		Stop();
+	}
+	else
+	{
+		Cleanup();
+	}
 }
 
 bool FProcessWrapper::Start(bool bWaitForExit)
diff --git a/Engine/Source/Programs/SubmitTool/Private/Logic/ProcessWrapper.h b/Engine/Source/Programs/SubmitTool/Private/Logic/ProcessWrapper.h
--- a/Engine/Source/Programs/SubmitTool/Private/Logic/ProcessWrapper.h
+++ b/Engine/Source/Programs/SubmitTool/Private/Logic/ProcessWrapper.h
@@ -17,7 +17,7 @@ DECLARE_DELEGATE_OneParam(FOnCompleted, const int32 /*Return code*/)
 class FProcessWrapper
 {
 public:
-	FProcessWrapper(const FString& InProcessName, const FString& InPath, const FString& InArgs, const FOnCompleted& InOnCompleted = FOnCompleted(), const FOnOutputLine& InOnOutputLine = FOnOutputLine(), const FString& InWorkingDir = FString(), const bool bLaunchHidden = true, const bool bLaunchReallyHidden = true, const bool bLaunchDetached = false);
+	FProcessWrapper(const FString& InProcessName, const FString& InPath, const FString& InArgs, const FOnCompleted& InOnCompleted = FOnCompleted(), const FOnOutputLine& InOnOutputLine = FOnOutputLine(), const FString& InWorkingDir = FString(), const bool bLaunchHidden = true, const bool bLaunchReallyHidden = true, const bool bLaunchDetached = false, bool InbOwnsProcessLifetime = true);
 	~FProcessWrapper();
 	bool Start(bool bWaitForExit = false);
 	void Stop();
@@ -44,6 +44,7 @@ private:
 	const bool bLaunchesHidden;
 	const bool bLaunchesReallyHidden;
 	const bool bLaunchDetached;
+	const bool bOwnsProcessLifetime;
 
 	TUniquePtr<FProcHandle> ProcessHandle;
 	TUniquePtr<FProcessPipes> Pipes;
diff --git a/Engine/Source/Programs/SubmitTool/Private/Logic/UpdateService.cpp b/Engine/Source/Programs/SubmitTool/Private/Logic/UpdateService.cpp
--- a/Engine/Source/Programs/SubmitTool/Private/Logic/UpdateService.cpp
+++ b/Engine/Source/Programs/SubmitTool/Private/Logic/UpdateService.cpp
@@ -111,7 +111,8 @@ void FUpdateService::StartAutoUpdateScript()
 			FString(),
 			bLaunchesHidden,
 			bLaunchesReallyHidden,
-			bLaunchDetached);
+			bLaunchDetached,
+			false);
 
 	ScriptProcess.Start();
 
