Reworked the filtering for files in hierarchy and improved performance of filepath testing

#rnx

diff --git a/Engine/Source/Programs/SubmitTool/Private/Logic/Validators/ValidatorBase.cpp b/Engine/Source/Programs/SubmitTool/Private/Logic/Validators/ValidatorBase.cpp
--- a/Engine/Source/Programs/SubmitTool/Private/Logic/Validators/ValidatorBase.cpp
+++ b/Engine/Source/Programs/SubmitTool/Private/Logic/Validators/ValidatorBase.cpp
@@ -10,6 +10,7 @@
 #include "Logic/TagService.h"
 #include "Logic/Services/Interfaces/ICacheDataService.h"
 #include "Logic/Services/Interfaces/ISTSourceControlService.h"
+#include "SubmitToolUtils.h"
 
 FValidatorBase::FValidatorBase(const FName& InNameId, const FSubmitToolParameters& InParameters, TWeakPtr<FSubmitToolServiceProvider> InServiceProvider, const FString& InDefinition) :
 	ValidatorNameID(InNameId),
@@ -66,6 +67,11 @@ void FValidatorBase::StartValidation()
 
 	if(bAppliesToCL)
 	{
+		if (FilteredFiles.Num() > 0)
+		{
+			UE_LOG(LogSubmitToolDebug, Log, TEXT("[%s] Validating files:\n - %s"), *ValidatorName, *FString::JoinBy(FilteredFiles, TEXT("\n - "), [](const FSourceControlStateRef& FileRef) { return FileRef->GetFilename(); }));
+		}
+
 		if (!bIsValidSetup)
 		{
 			// Try to recover if things have changed since startup
@@ -373,7 +379,6 @@ bool FValidatorBase::AppliesToFile(const FSourceControlStateRef InFile, bool Inb
 		// For validators that do not work on local files, we always apply
 		return true;
 	}
-	
 
 	if(!InFile->IsDeleted()
 		|| (InFile->IsDeleted() && Definition->bAcceptDeletedFiles))
@@ -437,27 +442,22 @@ bool FValidatorBase::AppliesToFile(const FSourceControlStateRef InFile, bool Inb
 			}
 		}
 
-		if(bIncluded && (!Definition->RequireFileInHierarchy.IsEmpty() || !Definition->ExcludeWhenFileInHierarchy.IsEmpty()))
+		if(bIncluded)
 		{
-			bIncluded = Definition->RequireFileInHierarchy.IsEmpty();
-			FString CurrentDir = FPaths::GetPath(Filename);
-
-			while (!CurrentDir.IsEmpty() && IFileManager::Get().DirectoryExists(*CurrentDir))
+			if (!Definition->RequireFileInHierarchy.IsEmpty())
 			{
-				TArray<FString> IncludeFilesFound;
-				IFileManager::Get().FindFiles(IncludeFilesFound, *(CurrentDir / Definition->RequireFileInHierarchy), true, false);
-
-				TArray<FString> ExcludeFilesFound;
-				IFileManager::Get().FindFiles(ExcludeFilesFound, *(CurrentDir / Definition->ExcludeWhenFileInHierarchy), true, false);
-
+				if (!FSubmitToolUtils::IsFileInHierarchy(Definition->RequireFileInHierarchy, Filename))
+				{
+					bIncluded = false;
+				}
+			}
 
-				if (IncludeFilesFound.Num() > 0 || ExcludeFilesFound.Num() > 0)
-				{		
-					bIncluded = ExcludeFilesFound.Num() == 0;
-					break;
+			if (!Definition->ExcludeWhenFileInHierarchy.IsEmpty())
+			{
+				if (FSubmitToolUtils::IsFileInHierarchy(Definition->ExcludeWhenFileInHierarchy, Filename))
+				{
+					bIncluded = false;
 				}
-				
-				CurrentDir = FPaths::GetPath(CurrentDir);
 			}
 		}
 
@@ -472,13 +472,11 @@ bool FValidatorBase::AppliesToFile(const FSourceControlStateRef InFile, bool Inb
 				bIncluded = false;
 			}
 		}
-
 	}
 
 	return bIncluded;
 }
 
-
 bool FValidatorBase::AppliesToCL(const FString& InCLDescription, const TArray<FSourceControlStateRef>& FilesInCL, const TArray<const FTag*>& Tags, TArray<FSourceControlStateRef>& OutFilteredFiles, TArray<FSourceControlStateRef>& OutIncrementalSkips, bool InbAllowIncremental) const
 {
 	for(const FSourceControlStateRef& File : FilesInCL)
diff --git a/Engine/Source/Programs/SubmitTool/Private/SubmitToolUtils.cpp b/Engine/Source/Programs/SubmitTool/Private/SubmitToolUtils.cpp
--- a/Engine/Source/Programs/SubmitTool/Private/SubmitToolUtils.cpp
+++ b/Engine/Source/Programs/SubmitTool/Private/SubmitToolUtils.cpp
@@ -6,6 +6,7 @@
 #include "Logging/SubmitToolLog.h"
 #include "Widgets/SWindow.h"
 #include "Framework/Application/SlateApplication.h"
+#include "HAL/FileManagerGeneric.h"
 
 #if PLATFORM_WINDOWS
 #include "Windows/AllowWindowsPlatformTypes.h"
@@ -15,6 +16,8 @@
 
 #define LOCTEXT_NAMESPACE "FSubmitToolUtils"
 
+TMap<FString, TMap<bool, TSet<FString>>> FSubmitToolUtils::HierarchyWildcardsCache;
+
 FString FSubmitToolUtils::GetLocalAppDataPath()
 {
 #if PLATFORM_WINDOWS
@@ -103,4 +106,50 @@ void FSubmitToolUtils::EnsureWindowIsInView(TSharedRef<SWindow> InWindow, bool b
 	}	
 }
 
+bool FSubmitToolUtils::IsFileInHierarchy(const FString& InWildcard, const FString& InPath)
+{
+	HierarchyWildcardsCache.FindOrAdd(InWildcard);	
+	for (const FString& PathWithFile : HierarchyWildcardsCache[InWildcard].FindOrAdd(true))
+	{			
+		if (FPaths::IsUnderDirectory(InPath, PathWithFile))
+		{
+			return true;
+		}
+	}
+	
+	FString CurrentDirectory = InPath;
+	if (!FPaths::GetExtension(InPath).IsEmpty())
+	{
+		CurrentDirectory = FPaths::GetPath(InPath);
+	}
+
+	TSet<FString>& NegativePaths = HierarchyWildcardsCache[InWildcard].FindOrAdd(false);
+
+	while (!CurrentDirectory.IsEmpty() && IFileManager::Get().DirectoryExists(*CurrentDirectory))
+	{
+		if (NegativePaths.Contains(CurrentDirectory))
+		{
+			CurrentDirectory = FPaths::GetPath(CurrentDirectory);
+		}
+		else
+		{
+			TArray<FString> FilesInFolder;
+			IFileManager::Get().FindFiles(FilesInFolder, *(CurrentDirectory / InWildcard), true, false);
+
+			const bool bContained = FilesInFolder.Num() != 0;
+
+			HierarchyWildcardsCache[InWildcard][bContained].Add(CurrentDirectory);
+
+			if(bContained)
+			{
+				return true;
+			}
+
+			CurrentDirectory = FPaths::GetPath(CurrentDirectory);
+		}
+	}
+
+	return false;
+}
+
 #undef LOCTEXT_NAMESPACE
diff --git a/Engine/Source/Programs/SubmitTool/Private/SubmitToolUtils.h b/Engine/Source/Programs/SubmitTool/Private/SubmitToolUtils.h
--- a/Engine/Source/Programs/SubmitTool/Private/SubmitToolUtils.h
+++ b/Engine/Source/Programs/SubmitTool/Private/SubmitToolUtils.h
@@ -15,4 +15,11 @@ public:
 	static void CopyDiagnosticFilesToClipboard(TConstArrayView<FString> Files);
 
 	static void EnsureWindowIsInView(TSharedRef<SWindow> InWindow, bool bSingleWindow);
+
+	static bool IsFileInHierarchy(const FString& InWildcard, const FString& InPath);
+
+	static void CacheWildcardToHierarchy(const FString& InWildcard, const FString& InPath);
+
+private:
+	static TMap<FString, TMap<bool, TSet<FString>>> HierarchyWildcardsCache;
 };
