Error when no lockdown files were downloaded instead of warning when only one has failed

diff --git a/Engine/Source/Programs/SubmitTool/Private/Logic/P4LockdownService.cpp b/Engine/Source/Programs/SubmitTool/Private/Logic/P4LockdownService.cpp
--- a/Engine/Source/Programs/SubmitTool/Private/Logic/P4LockdownService.cpp
+++ b/Engine/Source/Programs/SubmitTool/Private/Logic/P4LockdownService.cpp
@@ -32,7 +32,7 @@ FP4LockdownService::FP4LockdownService(const FP4LockdownParameters& InParameters
 			bool bResult = true;
 			for(TTask<bool>& Task : DownloadFileTasks)
 			{
-				bResult &= Task.GetResult();
+				bResult |= Task.GetResult();
 			}
 
 			return bResult;
@@ -54,7 +54,7 @@ bool FP4LockdownService::ArePathsInLockdown(const TArray<FString>& InPaths, bool
 
 		if(!ConfigFilesTask.IsCompleted())
 		{
-			UE_LOG(LogSubmitToolP4, Error, TEXT("Downloading config files from P4 timed out, hardlock status is not latest, will use cache."));
+			UE_LOG(LogSubmitToolP4, Error, TEXT("Downloading config files from P4 timed out, hardlock status is not latest, will revert to use cache."));
 		}
 	}
 
@@ -80,7 +80,7 @@ bool FP4LockdownService::ArePathsInLockdown(const TArray<FString>& InPaths, bool
 
 	if(!ConfigFilesTask.GetResult())
 	{
-		UE_LOG(LogSubmitToolP4, Warning, TEXT("Downloading config files from P4 had failures. Lockdown data won't be complete"));
+		UE_LOG(LogSubmitToolP4, Error, TEXT("Downloading config files from P4 failed. Lockdown data won't be available"));
 	}
 
 	ParseAllowListData();
