Fix issue of incorrect invalidation when the file list comes out of order from p4

diff --git a/Engine/Source/Programs/SubmitTool/Private/Logic/ChangelistService.cpp b/Engine/Source/Programs/SubmitTool/Private/Logic/ChangelistService.cpp
--- a/Engine/Source/Programs/SubmitTool/Private/Logic/ChangelistService.cpp
+++ b/Engine/Source/Programs/SubmitTool/Private/Logic/ChangelistService.cpp
@@ -259,16 +259,39 @@ void FChangelistService::RehydrateDataFromP4State()
 		CLDescription = ChangelistState->GetDescriptionText().ToString();
 	}
 
-	if(ChangelistState->GetFilesStates() != FilesInCL)
+
+	TArray<FSourceControlStateRef> PendingFileStates = ChangelistState->GetFilesStates();
+	Algo::StableSortBy(PendingFileStates, [](const FSourceControlStateRef& FileState) { return FileState->GetFilename(); });
+
+	bool bNewLocalFiles = false;
+	if (PendingFileStates.Num() == FilesInCL.Num())
+	{
+		for (size_t i = 0; i < PendingFileStates.Num(); ++i)
+		{
+			if (!PendingFileStates[i]->GetFilename().Equals(ShelvedFilesInCL[i]->GetFilename(), ESearchCase::IgnoreCase))
+			{
+				bNewLocalFiles = true;
+				break;
+			}
+		}
+	}
+	else
+	{
+		bNewLocalFiles = true;
+	}
+
+	if(bNewLocalFiles)
 	{
 		ChangeType |= ETaskArea::LocalFiles;
-		FilesInCL = ChangelistState->GetFilesStates();
+		FilesInCL = MoveTemp(PendingFileStates);
 		FilesDepotPaths.Reset(FilesInCL.Num());
 		UE_LOG(LogSubmitToolP4, Log, TEXT("CL %s files were updated outside of Submit Tool while it was open, Validation state has been reset"), *GetCLID());
 		PrintFilesAndShelvedFiles();
 	}
 
-	const TArray<FSourceControlStateRef> P4Shelf = ChangelistState->GetShelvedFilesStates();
+	TArray<FSourceControlStateRef> P4Shelf = ChangelistState->GetShelvedFilesStates();
+	Algo::StableSortBy(P4Shelf, [](const FSourceControlStateRef& FileState) { return FileState->GetFilename(); });
+
 	bool bNewShelf = false;
 	if (P4Shelf.Num() == ShelvedFilesInCL.Num())
 	{
@@ -289,7 +312,7 @@ void FChangelistService::RehydrateDataFromP4State()
 	if(bNewShelf)
 	{
 		ChangeType |= ETaskArea::ShelvedFiles;
-		ShelvedFilesInCL = ChangelistState->GetShelvedFilesStates();
+		ShelvedFilesInCL = MoveTemp(P4Shelf);
 		ShelvedFilesDepotPaths.Reset(ShelvedFilesInCL.Num());
 		UE_LOG(LogSubmitToolP4, Log, TEXT("CL %s shelved files were updated outside of Submit Tool while it was open, Validation state has been reset"), *GetCLID());
 		PrintFilesAndShelvedFiles();
@@ -301,7 +324,7 @@ void FChangelistService::RehydrateDataFromP4State()
 	}
 }
 
-bool FChangelistService::AreCLDescriptionsIdentical()
+bool FChangelistService::AreCLDescriptionsIdentical() const
 {
 	const FString NewDescription = ChangelistState->GetDescriptionText()
 		.ToString()
diff --git a/Engine/Source/Programs/SubmitTool/Private/Logic/ChangelistService.h b/Engine/Source/Programs/SubmitTool/Private/Logic/ChangelistService.h
--- a/Engine/Source/Programs/SubmitTool/Private/Logic/ChangelistService.h
+++ b/Engine/Source/Programs/SubmitTool/Private/Logic/ChangelistService.h
@@ -136,7 +136,7 @@ private:
 
 	void FetchChangelistCallback(const FSourceControlOperationRef& UpdateOperation, ECommandResult::Type Result);
 	void RehydrateDataFromP4State();
-	bool AreCLDescriptionsIdentical();
+	bool AreCLDescriptionsIdentical() const;
 	void PrintFilesAndShelvedFiles();
 
 	void PrintMessages(const FSourceControlResultInfo& ResultInfo) const;
