Fix missing include of StringOutputDevice for submit tool

diff --git a/Engine/Source/Programs/SubmitTool/Private/Logic/PreSubmitOperations/VirtualizationOperation.cpp b/Engine/Source/Programs/SubmitTool/Private/Logic/PreSubmitOperations/VirtualizationOperation.cpp
--- a/Engine/Source/Programs/SubmitTool/Private/Logic/PreSubmitOperations/VirtualizationOperation.cpp
+++ b/Engine/Source/Programs/SubmitTool/Private/Logic/PreSubmitOperations/VirtualizationOperation.cpp
@@ -6,6 +6,7 @@
 #include "Configuration/Configuration.h"
 #include "Logic/Validators/ValidatorDefinition.h"
 #include "Misc/Paths.h"
+#include "Misc/StringOutputDevice.h"
 #include "Models/ModelInterface.h"
 
 const TCHAR* LexToString(FVirtualizationOperation::EVirtualizationErrorCode ErrorCode)
diff --git a/Engine/Source/Programs/SubmitTool/Private/Logic/Validators/EditorCommandletValidator.cpp b/Engine/Source/Programs/SubmitTool/Private/Logic/Validators/EditorCommandletValidator.cpp
--- a/Engine/Source/Programs/SubmitTool/Private/Logic/Validators/EditorCommandletValidator.cpp
+++ b/Engine/Source/Programs/SubmitTool/Private/Logic/Validators/EditorCommandletValidator.cpp
@@ -4,6 +4,7 @@
 
 #include "Misc/FileHelper.h"
 #include "Misc/Paths.h"
+#include "Misc/StringOutputDevice.h"
 
 #include "CommandLine/CmdLineParameters.h"
 
diff --git a/Engine/Source/Programs/SubmitTool/Private/Logic/Validators/JsonValidator.cpp b/Engine/Source/Programs/SubmitTool/Private/Logic/Validators/JsonValidator.cpp
--- a/Engine/Source/Programs/SubmitTool/Private/Logic/Validators/JsonValidator.cpp
+++ b/Engine/Source/Programs/SubmitTool/Private/Logic/Validators/JsonValidator.cpp
@@ -7,6 +7,7 @@
 #include "Serialization/JsonSerializer.h"
 #include "Misc/FileHelper.h"
 #include "Misc/Paths.h"
+#include "Misc/StringOutputDevice.h"
 
 #include "Models/ModelInterface.h"
 
diff --git a/Engine/Source/Programs/SubmitTool/Private/Logic/Validators/PackageDataValidator.cpp b/Engine/Source/Programs/SubmitTool/Private/Logic/Validators/PackageDataValidator.cpp
--- a/Engine/Source/Programs/SubmitTool/Private/Logic/Validators/PackageDataValidator.cpp
+++ b/Engine/Source/Programs/SubmitTool/Private/Logic/Validators/PackageDataValidator.cpp
@@ -2,6 +2,7 @@
 
 #include "PackageDataValidator.h"
 
+#include "Misc/StringOutputDevice.h"
 #include "Models/ModelInterface.h"
 #include "UObject/PackageTrailer.h"
 #include "VirtualizationUtilities.h"
diff --git a/Engine/Source/Programs/SubmitTool/Private/Logic/Validators/PreflightValidator.cpp b/Engine/Source/Programs/SubmitTool/Private/Logic/Validators/PreflightValidator.cpp
--- a/Engine/Source/Programs/SubmitTool/Private/Logic/Validators/PreflightValidator.cpp
+++ b/Engine/Source/Programs/SubmitTool/Private/Logic/Validators/PreflightValidator.cpp
@@ -7,6 +7,7 @@
 #include "Logging/SubmitToolLog.h"
 #include "Internationalization/Regex.h"
 #include "Models/ModelInterface.h"
+#include "Misc/StringOutputDevice.h"
 #include "HAL/PlatformFileManager.h"
 #include "GenericPlatform/GenericPlatformFile.h"
 
diff --git a/Engine/Source/Programs/SubmitTool/Private/Logic/Validators/UBTValidator.cpp b/Engine/Source/Programs/SubmitTool/Private/Logic/Validators/UBTValidator.cpp
--- a/Engine/Source/Programs/SubmitTool/Private/Logic/Validators/UBTValidator.cpp
+++ b/Engine/Source/Programs/SubmitTool/Private/Logic/Validators/UBTValidator.cpp
@@ -4,6 +4,7 @@
 
 #include "Misc/FileHelper.h"
 #include "Misc/Paths.h"
+#include "Misc/StringOutputDevice.h"
 
 #include "CommandLine/CmdLineParameters.h"
 
diff --git a/Engine/Source/Programs/SubmitTool/Private/Logic/Validators/ValidatorBase.cpp b/Engine/Source/Programs/SubmitTool/Private/Logic/Validators/ValidatorBase.cpp
--- a/Engine/Source/Programs/SubmitTool/Private/Logic/Validators/ValidatorBase.cpp
+++ b/Engine/Source/Programs/SubmitTool/Private/Logic/Validators/ValidatorBase.cpp
@@ -5,6 +5,7 @@
 #include "Internationalization/Regex.h"
 #include "Misc/Paths.h"
 #include "Misc/Timespan.h"
+#include "Misc/StringOutputDevice.h"
 #include "AnalyticsEventAttribute.h"
 #include "Configuration/Configuration.h"
 #include "Logic/TagService.h"
diff --git a/Engine/Source/Programs/SubmitTool/Private/Logic/Validators/ValidatorRunExecutable.cpp b/Engine/Source/Programs/SubmitTool/Private/Logic/Validators/ValidatorRunExecutable.cpp
--- a/Engine/Source/Programs/SubmitTool/Private/Logic/Validators/ValidatorRunExecutable.cpp
+++ b/Engine/Source/Programs/SubmitTool/Private/Logic/Validators/ValidatorRunExecutable.cpp
@@ -11,6 +11,7 @@
 #include "Misc/FileHelper.h"
 #include "Misc/PathViews.h"
 #include "Misc/Paths.h"
+#include "Misc/StringOutputDevice.h"
 #include "Models/ModelInterface.h"
 #include "Models/SubmitToolUserPrefs.h"
 #include "Modules/BuildVersion.h"
