Removed RenderCore and ShaderCompiler modules from submit tool dependencies since they increased the program size by 20Mb for an unused validator.
Narrowed slate resource dependencies to what it's actually being used by submit tool, saved another 20Mb of packaged size
Fixing includes issues for non-unity and shipping configuration

diff --git a/Engine/Source/Programs/SubmitTool/Private/Logic/Validators/ShaderValidator.cpp b/Engine/Source/Programs/SubmitTool/Private/Logic/Validators/ShaderValidator.cpp
--- a/Engine/Source/Programs/SubmitTool/Private/Logic/Validators/ShaderValidator.cpp
+++ b/Engine/Source/Programs/SubmitTool/Private/Logic/Validators/ShaderValidator.cpp
@@ -2,16 +2,20 @@
 
 #include "ShaderValidator.h"
 
-#include "Misc/FileHelper.h"
-#include "Misc/Paths.h"
-#include "ShaderCompilerCommon.h"
+// This is commented out for the moment due to this validator not being in used and having a 20MB cost in the package size due to requiring the RenderCore module as a dependency.
+// it should eventually be split off as a plugin
 
-#include "CommandLine/CmdLineParameters.h"
 
-#include "Configuration/Configuration.h"
-#include "Logging/SubmitToolLog.h"
-#include "Models/ModelInterface.h"
-#include "Models/SubmitToolUserPrefs.h"
+//#include "Misc/FileHelper.h"
+//#include "Misc/Paths.h"
+//#include "ShaderCompilerCommon.h"
+//
+//#include "CommandLine/CmdLineParameters.h"
+//
+//#include "Configuration/Configuration.h"
+//#include "Logging/SubmitToolLog.h"
+//#include "Models/ModelInterface.h"
+//#include "Models/SubmitToolUserPrefs.h"
 
 
 FShaderValidator::FShaderValidator(const FName& InNameId, const FSubmitToolParameters& InParameters, TSharedRef<FSubmitToolServiceProvider> InServiceProvider, const FString& InDefinition) :
@@ -21,7 +25,7 @@ FShaderValidator::FShaderValidator(const FName& InNameId, const FSubmitToolParam
 
 void FShaderValidator::StartAsyncWork(const FString& InCLDescription, const TArray<FSourceControlStateRef>& InFilteredFilesInCL, const TArray<const FTag*>& InTags)
 {
-	this->StartAsyncTask([this, InFilteredFilesInCL](const UE::Tasks::FCancellationToken& InCancellationToken) -> bool
+	/*this->StartAsyncTask([this, InFilteredFilesInCL](const UE::Tasks::FCancellationToken& InCancellationToken) -> bool
 	{
 		bool bValid = true;
 
@@ -56,5 +60,5 @@ void FShaderValidator::StartAsyncWork(const FString& InCLDescription, const TArr
 		}
 
 		return bValid;
-	});
+	});*/
 }
diff --git a/Engine/Source/Programs/SubmitTool/Private/Models/SubmitToolUserPrefs.cpp b/Engine/Source/Programs/SubmitTool/Private/Models/SubmitToolUserPrefs.cpp
--- a/Engine/Source/Programs/SubmitTool/Private/Models/SubmitToolUserPrefs.cpp
+++ b/Engine/Source/Programs/SubmitTool/Private/Models/SubmitToolUserPrefs.cpp
@@ -3,6 +3,7 @@
 
 #include "SubmitToolUserPrefs.h"
 #include "Misc/Paths.h"
+#include "Misc/StringOutputDevice.h"
 #include "HAL/FileManager.h"
 
 FString FSubmitToolUserPrefs::FilePath = FString();
diff --git a/Engine/Source/Programs/SubmitTool/Private/Parameters/SubmitToolParametersBuilder.cpp b/Engine/Source/Programs/SubmitTool/Private/Parameters/SubmitToolParametersBuilder.cpp
--- a/Engine/Source/Programs/SubmitTool/Private/Parameters/SubmitToolParametersBuilder.cpp
+++ b/Engine/Source/Programs/SubmitTool/Private/Parameters/SubmitToolParametersBuilder.cpp
@@ -11,6 +11,7 @@
 #include "Logging/SubmitToolLog.h"
 #include "CommandLine/CmdLineParameters.h"
 #include "Misc/ConfigContext.h"
+#include "Misc/StringOutputDevice.h"
 
 FSubmitToolParametersBuilder::FSubmitToolParametersBuilder(const FString& InParametersXmlFile)
 {
diff --git a/Engine/Source/Programs/SubmitTool/Private/View/SubmitToolCommandHandler.cpp b/Engine/Source/Programs/SubmitTool/Private/View/SubmitToolCommandHandler.cpp
--- a/Engine/Source/Programs/SubmitTool/Private/View/SubmitToolCommandHandler.cpp
+++ b/Engine/Source/Programs/SubmitTool/Private/View/SubmitToolCommandHandler.cpp
@@ -10,7 +10,6 @@
 #include "Framework/Application/SlateApplication.h"
 #include "Framework/Commands/UICommandList.h"
 #include "Modules/ModuleManager.h"
-#include "ISlateReflectorModule.h"
 
 #define LOCTEXT_NAMESPACE "FSubmitToolCommandHandler"
 
diff --git a/Engine/Source/Programs/SubmitTool/SubmitTool.Build.cs b/Engine/Source/Programs/SubmitTool/SubmitTool.Build.cs
--- a/Engine/Source/Programs/SubmitTool/SubmitTool.Build.cs
+++ b/Engine/Source/Programs/SubmitTool/SubmitTool.Build.cs
@@ -21,8 +21,8 @@ public class SubmitTool : ModuleRules
 				"InputCore",
 				"ApplicationCore",
 				"Projects",
-				"RenderCore",
-				"ShaderCompilerCommon",
+				//"RenderCore", // Stripping both of these for now, this was only used by ShaderValidator but it added 20 mb to the packaged size for a validator that's not in use
+				//"ShaderCompilerCommon",
 				"Slate",
 				"SlateCore",
 				"StandaloneRenderer",
@@ -74,6 +74,37 @@ public class SubmitTool : ModuleRules
 			);
 		}
 		
-		RuntimeDependencies.Add("$(EngineDir)/Content/Editor/Slate/...*.png", StagedFileType.UFS);
+		RuntimeDependencies.Add("$(EngineDir)/Content/Editor/Slate/Common/...*.png", StagedFileType.UFS);
+		RuntimeDependencies.Add("$(EngineDir)/Content/Editor/Slate/Docking/...*.png", StagedFileType.UFS);
+		RuntimeDependencies.Add("$(EngineDir)/Content/Editor/Slate/Icons/Cross_12x.png", StagedFileType.UFS);
+		RuntimeDependencies.Add("$(EngineDir)/Content/Editor/Slate/Icons/cursor_cardinal_cross.png", StagedFileType.UFS);
+		RuntimeDependencies.Add("$(EngineDir)/Content/Editor/Slate/Icons/cursor_grab.png", StagedFileType.UFS);
+		RuntimeDependencies.Add("$(EngineDir)/Content/Editor/Slate/Icons/denied_16x.png", StagedFileType.UFS);
+		RuntimeDependencies.Add("$(EngineDir)/Content/Editor/Slate/Icons/Edit/icon_Edit_Duplicate_16x.png", StagedFileType.UFS);
+		RuntimeDependencies.Add("$(EngineDir)/Content/Editor/Slate/Icons/EditorAppIcon.png", StagedFileType.UFS);
+		RuntimeDependencies.Add("$(EngineDir)/Content/Editor/Slate/Icons/ellipsis_12x.png", StagedFileType.UFS);
+		RuntimeDependencies.Add("$(EngineDir)/Content/Editor/Slate/Icons/Empty_14x.png", StagedFileType.UFS);
+		RuntimeDependencies.Add("$(EngineDir)/Content/Editor/Slate/Icons/eyedropper_16px.png", StagedFileType.UFS);
+		RuntimeDependencies.Add("$(EngineDir)/Content/Editor/Slate/Icons/Help/icon_Help_Documentation_16x.png", StagedFileType.UFS);
+		RuntimeDependencies.Add("$(EngineDir)/Content/Editor/Slate/Icons/Help/icon_Help_support_16x.png", StagedFileType.UFS);
+		RuntimeDependencies.Add("$(EngineDir)/Content/Editor/Slate/Icons/icon_Downloads_16x.png", StagedFileType.UFS);
+		RuntimeDependencies.Add("$(EngineDir)/Content/Editor/Slate/Icons/icon_error_16x.png", StagedFileType.UFS);
+		RuntimeDependencies.Add("$(EngineDir)/Content/Editor/Slate/Icons/icon_help_16x.png", StagedFileType.UFS);
+		RuntimeDependencies.Add("$(EngineDir)/Content/Editor/Slate/Icons/icon_info_16x.png", StagedFileType.UFS);
+		RuntimeDependencies.Add("$(EngineDir)/Content/Editor/Slate/Icons/icon_redo_16px.png", StagedFileType.UFS);
+		RuntimeDependencies.Add("$(EngineDir)/Content/Editor/Slate/Icons/icon_tab_toolbar_16px.png", StagedFileType.UFS);
+		RuntimeDependencies.Add("$(EngineDir)/Content/Editor/Slate/Icons/icon_undo_16px.png", StagedFileType.UFS);
+		RuntimeDependencies.Add("$(EngineDir)/Content/Editor/Slate/Icons/icon_warning_16x.png", StagedFileType.UFS);
+		RuntimeDependencies.Add("$(EngineDir)/Content/Editor/Slate/Icons/PlusSymbol_12x.png", StagedFileType.UFS);
+		RuntimeDependencies.Add("$(EngineDir)/Content/Editor/Slate/Icons/refresh_12x.png", StagedFileType.UFS);
+		RuntimeDependencies.Add("$(EngineDir)/Content/Editor/Slate/Icons/Star_16x.png", StagedFileType.UFS);
+		RuntimeDependencies.Add("$(EngineDir)/Content/Editor/Slate/Icons/toolbar_expand_16x.png", StagedFileType.UFS);
+		RuntimeDependencies.Add("$(EngineDir)/Content/Editor/Slate/Old/HyperlinkDotted.png", StagedFileType.UFS);
+		RuntimeDependencies.Add("$(EngineDir)/Content/Editor/Slate/Old/HyperlinkUnderline.png", StagedFileType.UFS);
+		RuntimeDependencies.Add("$(EngineDir)/Content/Editor/Slate/Old/White.png", StagedFileType.UFS);
+		RuntimeDependencies.Add("$(EngineDir)/Content/Editor/Slate/Old/Menu_Background.png", StagedFileType.UFS); 
+		RuntimeDependencies.Add("$(EngineDir)/Content/Editor/Slate/Old/DashedBorder.png", StagedFileType.UFS); 
+		RuntimeDependencies.Add("$(EngineDir)/Content/Editor/Slate/Old/ToolTip_Background.png", StagedFileType.UFS);
+		RuntimeDependencies.Add("$(EngineDir)/Content/Editor/Slate/Checkerboard.png", StagedFileType.UFS);
 	}
 }
diff --git a/Engine/Source/Programs/SubmitTool/SubmitTool.Target.cs b/Engine/Source/Programs/SubmitTool/SubmitTool.Target.cs
--- a/Engine/Source/Programs/SubmitTool/SubmitTool.Target.cs
+++ b/Engine/Source/Programs/SubmitTool/SubmitTool.Target.cs
@@ -40,9 +40,8 @@ public class SubmitToolTarget : TargetRules
 		WindowsPlatform.bMergeIdenticalCOMDATs = true;
 		WindowsPlatform.bUseBundledDbgHelp = false;
 		WindowsPlatform.bPixProfilingEnabled = false;
-		
+
 		GlobalDefinitions.Add("AUTOSDKS_ENABLED=0");
-		GlobalDefinitions.Add("UE_TRACE_AUTOSTART=1");
 		GlobalDefinitions.Add("EXCLUDE_NONPAK_UE_EXTENSIONS=0");
 
 		OptedInModulePlatforms = new UnrealTargetPlatform[] { UnrealTargetPlatform.Win64, UnrealTargetPlatform.Mac,
