Added button to open preflight links directly on the tag widget
Waiting to start a preflight is now opt-in per validator, by default preflights can be started at any point without restrictions

diff --git a/Engine/Source/Programs/SubmitTool/Private/Logic/Validators/ValidatorDefinition.h b/Engine/Source/Programs/SubmitTool/Private/Logic/Validators/ValidatorDefinition.h
--- a/Engine/Source/Programs/SubmitTool/Private/Logic/Validators/ValidatorDefinition.h
+++ b/Engine/Source/Programs/SubmitTool/Private/Logic/Validators/ValidatorDefinition.h
@@ -208,6 +208,9 @@ struct FValidatorDefinition
 	/// 
 	UPROPERTY()
 	ETaskArea TaskArea = ETaskArea::Everything;
+
+	UPROPERTY()
+	bool bBlocksPreflightStart = false;
 };
 
 USTRUCT()
diff --git a/Engine/Source/Programs/SubmitTool/Private/Models/ModelInterface.cpp b/Engine/Source/Programs/SubmitTool/Private/Models/ModelInterface.cpp
--- a/Engine/Source/Programs/SubmitTool/Private/Models/ModelInterface.cpp
+++ b/Engine/Source/Programs/SubmitTool/Private/Models/ModelInterface.cpp
@@ -12,6 +12,7 @@
 #include "SourceControlOperations.h"
 #include "HAL/PlatformApplicationMisc.h"
 #include "Async/Async.h"
+#include "Misc/StringOutputDevice.h"
 
 #include "Logic/Validators/PreflightValidator.h"
 #include "Logic/Validators/TagValidator.h"
@@ -339,7 +340,7 @@ bool FModelInterface::CanLaunchPreflight() const
 		if (Validator.IsValid())
 		{
 			const TSharedPtr<const FValidatorBase> Pinned = Validator.Pin();
-			if ((Pinned->Definition->TaskArea & ETaskArea::ShelveAndLocalFiles) != ETaskArea::None)
+			if (Pinned->Definition->bBlocksPreflightStart)
 			{
 				if ((Pinned->Definition->IsRequired && !Pinned->GetHasPassed()) || (!Pinned->Definition->IsRequired && Pinned->GetIsRunningOrQueued()))
 				{
diff --git a/Engine/Source/Programs/SubmitTool/Private/View/SubmitToolStyle.cpp b/Engine/Source/Programs/SubmitTool/Private/View/SubmitToolStyle.cpp
--- a/Engine/Source/Programs/SubmitTool/Private/View/SubmitToolStyle.cpp
+++ b/Engine/Source/Programs/SubmitTool/Private/View/SubmitToolStyle.cpp
@@ -142,6 +142,7 @@ FSubmitToolStyle::FSubmitToolStyle() : FSlateStyleSet(StyleName)
 
 	Set("AppIcon.DocumentationHelp", new IMAGE_BRUSH("Icons/Help/icon_Help_Documentation_16x", FVector2D(16, 16)));
 	Set("AppIcon.Refresh", new IMAGE_BRUSH("Icons/refresh_12x", FVector2D(12, 12)));
+	Set("AppIcon.OpenExternal", new IMAGE_BRUSH("Icons/Help/icon_Help_support_16x", FVector2D(12, 12)));
 	Set("AppIcon.Star16", new IMAGE_BRUSH("Icons/Star_16x", FVector2D(16, 16)));
 	
 	FSlateStyleRegistry::RegisterSlateStyle(*this);
diff --git a/Engine/Source/Programs/SubmitTool/Private/View/Widgets/TagWidget.cpp b/Engine/Source/Programs/SubmitTool/Private/View/Widgets/TagWidget.cpp
--- a/Engine/Source/Programs/SubmitTool/Private/View/Widgets/TagWidget.cpp
+++ b/Engine/Source/Programs/SubmitTool/Private/View/Widgets/TagWidget.cpp
@@ -436,6 +436,44 @@ void STagWidget::Construct(const FArguments& InArgs)
 					]					
 				]
 			];
+			
+		HorizontalBox->AddSlot()
+			.Padding(3, 0, 0, 0)
+			.AutoWidth()
+			[
+				SNew(SBox)
+				.MaxDesiredWidth(35)
+				[
+					SNew(SButton)
+					.ContentPadding(-5)
+					.IsEnabled_Lambda([this]{ return FModelInterface::GetInputEnabled() && Tag->GetValues().Num() > 0; })
+					.IsFocusable(true)
+					.ToolTipText(FText::FromString(TEXT("Open in browser")))
+					.OnClicked_Lambda([this]()
+					{
+						for (const TSharedPtr<FPreflightData>& PFData : *PreflightListUI)
+						{
+							if (Tag->GetValues().ContainsByPredicate([&PFData](const FString& InTagValue) { return InTagValue.Contains(PFData->ID); }))
+							{
+								FPlatformProcess::LaunchURL(*FString::Printf(TEXT("%sjob/%s"), *ModelInterface->GetParameters().HordeParameters.HordeServerAddress, *PFData->ID), nullptr, nullptr);
+							}
+						}
+
+						return FReply::Handled();
+					})
+					[
+						SNew(SBox)
+						.VAlign(VAlign_Center)
+						.HAlign(HAlign_Center)
+						.MaxDesiredHeight(12)
+						.MaxDesiredWidth(12)
+						[
+							SNew(SImage)
+							.Image(FSubmitToolStyle::Get().GetBrush("AppIcon.OpenExternal"))
+						]
+					]
+				]
+			];
 
 		Contents.Get()->AddSlot()
 			.AutoHeight()
