Fixed issue where ordering of files coming from p4 would make submit tool think that the CL contents have changed
Optimized evaluation of CL applicability for non file validators

diff --git a/Engine/Source/Programs/SubmitTool/Private/Logic/ChangelistService.cpp b/Engine/Source/Programs/SubmitTool/Private/Logic/ChangelistService.cpp
--- a/Engine/Source/Programs/SubmitTool/Private/Logic/ChangelistService.cpp
+++ b/Engine/Source/Programs/SubmitTool/Private/Logic/ChangelistService.cpp
@@ -260,15 +260,16 @@ void FChangelistService::RehydrateDataFromP4State()
 	}
 
 
-	TArray<FSourceControlStateRef> PendingFileStates = ChangelistState->GetFilesStates();
-	Algo::StableSortBy(PendingFileStates, [](const FSourceControlStateRef& FileState) { return FileState->GetFilename(); });
-
 	bool bNewLocalFiles = false;
-	if (PendingFileStates.Num() == FilesInCL.Num())
+	if (ChangelistState->GetFilesStatesNum() == FilesInCL.Num())
 	{
+		const TArray<FSourceControlStateRef>& PendingFileStates = ChangelistState->GetFilesStates();
+		Algo::StableSortBy(PendingFileStates, [](const FSourceControlStateRef& FileState) { return FileState->GetFilename(); });
+		Algo::StableSortBy(FilesInCL, [](const FSourceControlStateRef& FileState) { return FileState->GetFilename(); });
+
 		for (size_t i = 0; i < PendingFileStates.Num(); ++i)
 		{
-			if (!PendingFileStates[i]->GetFilename().Equals(ShelvedFilesInCL[i]->GetFilename(), ESearchCase::IgnoreCase))
+			if (!PendingFileStates[i]->GetFilename().Equals(FilesInCL[i]->GetFilename(), ESearchCase::IgnoreCase))
 			{
 				bNewLocalFiles = true;
 				break;
@@ -282,19 +283,23 @@ void FChangelistService::RehydrateDataFromP4State()
 
 	if(bNewLocalFiles)
 	{
+		UE_LOG(LogSubmitToolP4, Log, TEXT("CL %s files were updated outside of Submit Tool while it was open, Validation state has been reset"), *GetCLID());
+		UE_LOG(LogSubmitToolP4Debug, Log, TEXT("Current ST File info:"));
+		PrintFilesAndShelvedFiles();
 		ChangeType |= ETaskArea::LocalFiles;
-		FilesInCL = MoveTemp(PendingFileStates);
+		FilesInCL = ChangelistState->GetFilesStates();
 		FilesDepotPaths.Reset(FilesInCL.Num());
-		UE_LOG(LogSubmitToolP4, Log, TEXT("CL %s files were updated outside of Submit Tool while it was open, Validation state has been reset"), *GetCLID());
+		UE_LOG(LogSubmitToolP4Debug, Log, TEXT("New P4 File info:"));
 		PrintFilesAndShelvedFiles();
 	}
 
-	TArray<FSourceControlStateRef> P4Shelf = ChangelistState->GetShelvedFilesStates();
-	Algo::StableSortBy(P4Shelf, [](const FSourceControlStateRef& FileState) { return FileState->GetFilename(); });
-
 	bool bNewShelf = false;
-	if (P4Shelf.Num() == ShelvedFilesInCL.Num())
+	if (ChangelistState->GetShelvedFilesStatesNum() == ShelvedFilesInCL.Num())
 	{
+		const TArray<FSourceControlStateRef>& P4Shelf = ChangelistState->GetShelvedFilesStates();
+		Algo::StableSortBy(P4Shelf, [](const FSourceControlStateRef& FileState) { return FileState->GetFilename(); });
+		Algo::StableSortBy(ShelvedFilesInCL, [](const FSourceControlStateRef& FileState) { return FileState->GetFilename(); });
+
 		for (size_t i = 0; i < P4Shelf.Num(); ++i)
 		{
 			if (!P4Shelf[i]->GetFilename().Equals(ShelvedFilesInCL[i]->GetFilename(), ESearchCase::IgnoreCase))
@@ -311,10 +316,13 @@ void FChangelistService::RehydrateDataFromP4State()
 
 	if(bNewShelf)
 	{
+		UE_LOG(LogSubmitToolP4, Log, TEXT("CL %s shelved files were updated outside of Submit Tool while it was open, Validation state has been reset"), *GetCLID());
+		UE_LOG(LogSubmitToolP4Debug, Log, TEXT("Current ST File info:"));
+		PrintFilesAndShelvedFiles();
 		ChangeType |= ETaskArea::ShelvedFiles;
-		ShelvedFilesInCL = MoveTemp(P4Shelf);
+		ShelvedFilesInCL = ChangelistState->GetShelvedFilesStates();
 		ShelvedFilesDepotPaths.Reset(ShelvedFilesInCL.Num());
-		UE_LOG(LogSubmitToolP4, Log, TEXT("CL %s shelved files were updated outside of Submit Tool while it was open, Validation state has been reset"), *GetCLID());
+		UE_LOG(LogSubmitToolP4Debug, Log, TEXT("New P4 File info:"));
 		PrintFilesAndShelvedFiles();
 	}
 
diff --git a/Engine/Source/Programs/SubmitTool/Private/Logic/Validators/ValidatorBase.cpp b/Engine/Source/Programs/SubmitTool/Private/Logic/Validators/ValidatorBase.cpp
--- a/Engine/Source/Programs/SubmitTool/Private/Logic/Validators/ValidatorBase.cpp
+++ b/Engine/Source/Programs/SubmitTool/Private/Logic/Validators/ValidatorBase.cpp
@@ -480,6 +480,12 @@ bool FValidatorBase::AppliesToFile(const FSourceControlStateRef InFile, bool Inb
 
 bool FValidatorBase::AppliesToCL(const FString& InCLDescription, const TArray<FSourceControlStateRef>& FilesInCL, const TArray<const FTag*>& Tags, TArray<FSourceControlStateRef>& OutFilteredFiles, TArray<FSourceControlStateRef>& OutIncrementalSkips, bool InbAllowIncremental) const
 {
+	if ((Definition->TaskArea & ETaskArea::LocalFiles) == ETaskArea::None)
+	{
+		// For validators that do not work on files, we always apply
+		return true;
+	}
+
 	for(const FSourceControlStateRef& File : FilesInCL)
 	{
 		bool bIsIncrementalSkip;
