Splitted the preflight start functionality between Start and Queue buttons

diff --git a/Engine/Source/Programs/SubmitTool/Private/Logic/Validators/ValidatorDefinition.h b/Engine/Source/Programs/SubmitTool/Private/Logic/Validators/ValidatorDefinition.h
--- a/Engine/Source/Programs/SubmitTool/Private/Logic/Validators/ValidatorDefinition.h
+++ b/Engine/Source/Programs/SubmitTool/Private/Logic/Validators/ValidatorDefinition.h
@@ -210,7 +210,7 @@ struct FValidatorDefinition
 	ETaskArea TaskArea = ETaskArea::Everything;
 
 	UPROPERTY()
-	bool bBlocksPreflightStart = false;
+	bool bBlocksPreflightStart = true;
 };
 
 USTRUCT()
diff --git a/Engine/Source/Programs/SubmitTool/Private/Models/ModelInterface.cpp b/Engine/Source/Programs/SubmitTool/Private/Models/ModelInterface.cpp
--- a/Engine/Source/Programs/SubmitTool/Private/Models/ModelInterface.cpp
+++ b/Engine/Source/Programs/SubmitTool/Private/Models/ModelInterface.cpp
@@ -424,10 +424,11 @@ bool FModelInterface::GetInputEnabled()
 	return SubmitToolState == ESubmitToolAppState::WaitingUserInput || SubmitToolState == ESubmitToolAppState::SubmitLocked;
 }
 
-void FModelInterface::RequestPreflight()
+void FModelInterface::RequestPreflight(bool bForceStart)
 {
-	if (CanLaunchPreflight())
+	if (bForceStart || CanLaunchPreflight())
 	{
+		bPreflightQueued = false;
 		PreflightService->RequestPreflight();
 	}
 	else
diff --git a/Engine/Source/Programs/SubmitTool/Private/Models/ModelInterface.h b/Engine/Source/Programs/SubmitTool/Private/Models/ModelInterface.h
--- a/Engine/Source/Programs/SubmitTool/Private/Models/ModelInterface.h
+++ b/Engine/Source/Programs/SubmitTool/Private/Models/ModelInterface.h
@@ -173,7 +173,7 @@ public:
 	bool IsSubmitBlocked() const													{ return SubmitToolState == ESubmitToolAppState::SubmitLocked; }
 	void DeleteShelvedFiles() const													{ ChangelistService->DeleteShelvedFiles(DeleteShelveCallback); }
 
-	void RequestPreflight();
+	void RequestPreflight(bool bForceStart = false);
 	bool IsPreflightRequestInProgress() const										{ return PreflightService->IsRequestInProgress(); }
 	bool IsPreflightQueued() const													{ return bPreflightQueued; }
 	void RefreshPreflightInformation() const										{ PreflightService->FetchPreflightInfo(); UE_LOG(LogSubmitTool, Log, TEXT("Requesting preflight information...")); }
diff --git a/Engine/Source/Programs/SubmitTool/Private/View/Widgets/TagWidget.cpp b/Engine/Source/Programs/SubmitTool/Private/View/Widgets/TagWidget.cpp
--- a/Engine/Source/Programs/SubmitTool/Private/View/Widgets/TagWidget.cpp
+++ b/Engine/Source/Programs/SubmitTool/Private/View/Widgets/TagWidget.cpp
@@ -81,7 +81,7 @@ void STagWidget::Construct(const FArguments& InArgs)
 		.OnClicked_Raw(this, &STagWidget::OnLabelClick)
 		[
 			SNew(STextBlock)
-			.MinDesiredWidth(70)
+			.MinDesiredWidth(60)
 			.ColorAndOpacity_Lambda([this](){ 
 			switch(Tag->GetTagState())
 			{			
@@ -212,7 +212,7 @@ void STagWidget::Construct(const FArguments& InArgs)
 				.AutoWidth()
 				[
 					SNew(SBox)
-					.MinDesiredWidth(70)
+					.MinDesiredWidth(60)
 					[
 						SNew(SButton)
 						.HAlign(HAlign_Center)
@@ -254,10 +254,10 @@ void STagWidget::Construct(const FArguments& InArgs)
 		.AutoWidth()
 		[
 			SNew(SBox)
-			.MinDesiredWidth(70)
+			.MinDesiredWidth(60)
 			[
 				SNew(SSelectSourceControlUserWidget)
-				.ButtonText(FText::FromString("Users"))
+				.ButtonText(FText::FromString("User"))
 				.ModelInterface(ModelInterface)
 				.Tag(this->Tag)
 				.IsEnabled_Lambda([this]() { return FModelInterface::GetInputEnabled();})
@@ -269,10 +269,10 @@ void STagWidget::Construct(const FArguments& InArgs)
 		.AutoWidth()
 		[
 			SNew(SBox)
-			.MinDesiredWidth(70)
+			.MinDesiredWidth(60)
 			[
 				SNew(SSelectSourceControlGroupWidget)
-				.ButtonText(FText::FromString("Groups"))
+				.ButtonText(FText::FromString("Group"))
 				.ModelInterface(ModelInterface)
 				.Tag(this->Tag)
 				.IsEnabled_Lambda([this]() { return FModelInterface::GetInputEnabled();})
@@ -297,7 +297,7 @@ void STagWidget::Construct(const FArguments& InArgs)
 			.AutoWidth()
 			[
 				SNew(SBox)
-				.MinDesiredWidth(70)
+				.MinDesiredWidth(60)
 				[
 					SNew(SComboButton)
 					.IsEnabled(this->SelectValues.Num() > 0)
@@ -361,7 +361,7 @@ void STagWidget::Construct(const FArguments& InArgs)
 			.AutoWidth()
 			[
 				SNew(SBox)
-				.MinDesiredWidth(70)
+				.MinDesiredWidth(60)
 				[
 					SNew(SButton)
 					.OnClicked(this, &STagWidget::OnJiraClick)
@@ -418,7 +418,7 @@ void STagWidget::Construct(const FArguments& InArgs)
 			.AutoWidth()
 			[
 				SNew(SBox)
-				.MinDesiredWidth(70)
+				.MinDesiredWidth(60)
 				[
 					SNew(SComboButton)
 					.IsEnabled_Lambda([this](){ return PreflightListUI->Num() > 0; })
@@ -540,53 +540,78 @@ void STagWidget::Construct(const FArguments& InArgs)
 						];
 				})					
 			];
-
+			
 		HorizontalBox->AddSlot()
 			.Padding(3, 0, 0, 0)
 			.AutoWidth()
 			[
 				SNew(SBox)
-				.MinDesiredWidth(70)
+				.MinDesiredWidth(60)
 				[
 					SNew(SButton)
-					.OnClicked(this, &STagWidget::OnPreflightClick)
-					.IsEnabled_Lambda([&ModelInterface = ModelInterface] { return !ModelInterface->IsPreflightQueued() && !ModelInterface->IsPreflightRequestInProgress(); })
+					.OnClicked_Lambda([this]{ ModelInterface->RequestPreflight(); return FReply::Handled(); })
+					.IsEnabled_Lambda([&ModelInterface = ModelInterface] { return FModelInterface::GetInputEnabled() && !ModelInterface->IsPreflightQueued() && !ModelInterface->IsPreflightRequestInProgress(); })
 					.HAlign(HAlign_Center)
 					.VAlign(VAlign_Center)
 					[
 						SNew(STextBlock)
 						.Justification(ETextJustify::Center)
 						.Font(FAppStyle::Get().GetFontStyle("SmallFont"))
-						.Text_Lambda([&ModelInterface = ModelInterface]() 
-						{
-							if(ModelInterface->IsPreflightRequestInProgress())
-							{
-								return FText::FromString(TEXT("Requesting..."));
-							}
-							else if(ModelInterface->IsPreflightQueued())
-							{
-								return FText::FromString(TEXT("Queued"));
-							}
-							else
-							{
-								return FText::FromString(TEXT("Start"));
-							}
-						})
-						.ToolTipText_Lambda([&ModelInterface = ModelInterface]() 
-						{ 
-							if(ModelInterface->IsPreflightRequestInProgress())
-							{
-								return FText::FromString(TEXT("Preflight request in progress..."));
-							}
-							else if(ModelInterface->IsPreflightQueued())
+						.Text_Lambda([&ModelInterface = ModelInterface]()
 							{
-								return FText::FromString(TEXT("Preflight will start once the CL has run local validations"));
-							}
-							else
+								if (ModelInterface->IsPreflightRequestInProgress())
+								{
+									return FText::FromString(TEXT("Requesting..."));
+								}
+								else if (ModelInterface->IsPreflightQueued())
+								{
+									return FText::FromString(TEXT("Queued"));
+								}
+								else
+								{
+									return FText::FromString(TEXT("Queue"));
+								}
+							})
+						.ToolTipText_Lambda([&ModelInterface = ModelInterface]()
 							{
-								return FText::FromString(TEXT("Starts a Preflight in the browser"));
-							}
-						})
+								if (ModelInterface->IsPreflightRequestInProgress())
+								{
+									return FText::FromString(TEXT("Preflight request in progress..."));
+								}
+								else if (ModelInterface->IsPreflightQueued())
+								{
+									return FText::FromString(TEXT("Preflight will start once the CL has run local validations"));
+								}
+								else
+								{
+									return FText::FromString(TEXT("Enqueues a Preflight for when local validations have finished"));
+								}
+							})
+					]
+				]
+			];
+			
+		HorizontalBox->AddSlot()
+			.Padding(3, 0, 0, 0)
+			.AutoWidth()
+			[
+				SNew(SBox)
+				.MinDesiredWidth(60)
+				[
+					SNew(SButton)
+					.OnClicked_Lambda([this]{
+						ModelInterface->RequestPreflight(true); 
+						return FReply::Handled();
+					})
+					.IsEnabled_Lambda([&ModelInterface = ModelInterface] { return FModelInterface::GetInputEnabled() && !ModelInterface->IsPreflightRequestInProgress(); })
+					.HAlign(HAlign_Center)
+					.VAlign(VAlign_Center)
+					[
+						SNew(STextBlock)
+						.Justification(ETextJustify::Center)
+						.Font(FAppStyle::Get().GetFontStyle("SmallFont"))
+						.Text(FText::FromString(TEXT("Start")))
+						.ToolTipText(FText::FromString(TEXT("Starts a Preflight in Horde")))
 					]
 				]
 			];
