Fix for nullptr file checks

diff --git a/Engine/Source/Programs/SubmitTool/Private/Models/SubmitToolUserPrefs.cpp b/Engine/Source/Programs/SubmitTool/Private/Models/SubmitToolUserPrefs.cpp
--- a/Engine/Source/Programs/SubmitTool/Private/Models/SubmitToolUserPrefs.cpp
+++ b/Engine/Source/Programs/SubmitTool/Private/Models/SubmitToolUserPrefs.cpp
@@ -17,11 +17,18 @@ FSubmitToolUserPrefs::~FSubmitToolUserPrefs()
 		FSubmitToolUserPrefs::StaticStruct()->ExportText(OutputText, Instance, Instance, nullptr, PPF_None, nullptr);
 
 		FArchive* File = IFileManager::Get().CreateFileWriter(*FilePath, EFileWrite::FILEWRITE_EvenIfReadOnly);
-		*File << OutputText;
-		File->Close();
-		delete File;
-		File = nullptr;
-		UE_LOG(LogSubmitToolDebug, Verbose, TEXT("Saved User Prefs to %s:\n%s"), *FilePath, *OutputText);
+		if (File != nullptr)
+		{
+			*File << OutputText;
+			File->Close();
+			delete File;
+			File = nullptr;
+			UE_LOG(LogSubmitToolDebug, Verbose, TEXT("Saved User Prefs to %s:\n%s"), *FilePath, *OutputText);
+		}
+		else
+		{
+			UE_LOG(LogSubmitToolDebug, Verbose, TEXT("Failed to create a file for User Prefs with path %s:"), *FilePath);
+		}
 		Instance = nullptr;
 	}
 }
@@ -34,25 +41,32 @@ TUniquePtr<FSubmitToolUserPrefs> FSubmitToolUserPrefs::Initialize(const FString&
 	{
 		FString InText;
 		FArchive* File = IFileManager::Get().CreateFileReader(*FilePath, EFileRead::FILEREAD_None);
-		*File << InText;
-		File->Close();
-		delete File;
-		File = nullptr;
-		FStringOutputDevice Errors;
-		FSubmitToolUserPrefs::StaticStruct()->ImportText(*InText, NewPrefs.Get(), nullptr, 0, &Errors, FSubmitToolUserPrefs::StaticStruct()->GetName());
-
-		if(Instance != nullptr)
+		if (File != nullptr)
 		{
-			UE_LOG(LogSubmitTool, Warning, TEXT("UserPrefs have been reloaded"));
-		}
+			*File << InText;
+			File->Close();
+			delete File;
+			File = nullptr;
+			FStringOutputDevice Errors;
+			FSubmitToolUserPrefs::StaticStruct()->ImportText(*InText, NewPrefs.Get(), nullptr, 0, &Errors, FSubmitToolUserPrefs::StaticStruct()->GetName());
 
-		if(!Errors.IsEmpty())
-		{
-			UE_LOG(LogSubmitTool, Error, TEXT("Error loading User prefs file %s, using defaults"), *Errors);
+			if(!Errors.IsEmpty())
+			{
+				UE_LOG(LogSubmitTool, Error, TEXT("Error loading User prefs file %s, using defaults"), *Errors);
+			}
+			else
+			{
+				UE_LOG(LogSubmitToolDebug, Verbose, TEXT("Loaded User Prefs from %s:\n%s"), *FilePath, *InText);
+			}
 		}
 		else
 		{
-			UE_LOG(LogSubmitToolDebug, Verbose, TEXT("Loaded User Prefs from %s:\n%s"), *FilePath, *InText);
+			UE_LOG(LogSubmitToolDebug, Verbose, TEXT("Failed to open file for reading User Prefs %s:"), *FilePath);
+		}
+
+		if(Instance != nullptr)
+		{
+			UE_LOG(LogSubmitTool, Warning, TEXT("UserPrefs have been reloaded"));
 		}
 	}
 	else
