Bumping submit tool version
Adding option to okforgithub tag
Enabling UAT validators

diff --git a/Engine/Source/Programs/SubmitTool/Private/Version/AppVersionDefines.h b/Engine/Source/Programs/SubmitTool/Private/Version/AppVersionDefines.h
--- a/Engine/Source/Programs/SubmitTool/Private/Version/AppVersionDefines.h
+++ b/Engine/Source/Programs/SubmitTool/Private/Version/AppVersionDefines.h
@@ -20,7 +20,7 @@
 #endif
 
 #define SUBMIT_TOOL_VERSION_MAJOR 6
-#define SUBMIT_TOOL_VERSION_MINOR 0
+#define SUBMIT_TOOL_VERSION_MINOR 1
 #define SUBMIT_TOOL_VERSION_PATCH 0
 
 #ifndef SUBMIT_TOOL_CHANGELIST
