Fixed alignment of tag UI section

diff --git a/Engine/Source/Programs/SubmitTool/Private/View/Widgets/TagWidget.cpp b/Engine/Source/Programs/SubmitTool/Private/View/Widgets/TagWidget.cpp
--- a/Engine/Source/Programs/SubmitTool/Private/View/Widgets/TagWidget.cpp
+++ b/Engine/Source/Programs/SubmitTool/Private/View/Widgets/TagWidget.cpp
@@ -81,7 +81,7 @@ void STagWidget::Construct(const FArguments& InArgs)
 		.OnClicked_Raw(this, &STagWidget::OnLabelClick)
 		[
 			SNew(STextBlock)
-			.MinDesiredWidth(60)
+			.MinDesiredWidth(70)
 			.ColorAndOpacity_Lambda([this](){ 
 			switch(Tag->GetTagState())
 			{			
