Fixed issues with SubProjects not having the right arguments

diff --git a/Engine/Source/Programs/SubmitTool/Private/Logic/Validators/EditorCommandletValidator.cpp b/Engine/Source/Programs/SubmitTool/Private/Logic/Validators/EditorCommandletValidator.cpp
--- a/Engine/Source/Programs/SubmitTool/Private/Logic/Validators/EditorCommandletValidator.cpp
+++ b/Engine/Source/Programs/SubmitTool/Private/Logic/Validators/EditorCommandletValidator.cpp
@@ -125,12 +125,12 @@ void FEditorCommandletValidator::GetEditorsForPaths(const TArray<FSourceControlS
 
 			if (SubProjectJson->HasTypedField(FProjectEditorRecord::EpicAppProperty, EJson::String))
 			{
-				EditorParams.EditorArguments.Append(TEXT(" -epicapp ") + SubProjectJson->GetStringField(FProjectEditorRecord::EpicAppProperty));
+				EditorParams.EditorArguments.Append(TEXT(" -epicapp=") + SubProjectJson->GetStringField(FProjectEditorRecord::EpicAppProperty));
 			}
 
 			if (SubProjectJson->HasTypedField(FProjectEditorRecord::BaseDirProperty, EJson::String))
 			{
-				EditorParams.EditorArguments.Append(TEXT(" -BaseDir ") + SubProjectJson->GetStringField(FProjectEditorRecord::BaseDirProperty));
+				EditorParams.EditorArguments.Append(TEXT(" -BaseDir=") + SubProjectJson->GetStringField(FProjectEditorRecord::BaseDirProperty));
 			}
 
 			EditorParams.EditorExePath = SubProjectJson->GetStringField(FProjectEditorRecord::EngineLocationProperty);
