Added refresh button to Jira widget
Removed useless logging from "SourceControl"

diff --git a/Engine/Source/Programs/SubmitTool/Private/Models/ModelInterface.h b/Engine/Source/Programs/SubmitTool/Private/Models/ModelInterface.h
--- a/Engine/Source/Programs/SubmitTool/Private/Models/ModelInterface.h
+++ b/Engine/Source/Programs/SubmitTool/Private/Models/ModelInterface.h
@@ -191,9 +191,9 @@ public:
 
 	TSharedRef<FSubmitToolServiceProvider> GetServiceProvider() { return ServiceProvider.ToSharedRef(); }
 
-	TSharedPtr<FJiraService> GetJiraService() { return this->JiraService; }
-	TSharedPtr<FSwarmService> GetSwarmService() { return this->SwarmService; }
-	TSharedPtr<FPreflightService> GetPreflightService() { return this->PreflightService; }
+	TWeakPtr<FJiraService> GetJiraService() { return this->JiraService; }
+	TWeakPtr<FSwarmService> GetSwarmService() { return this->SwarmService; }
+	TWeakPtr<FPreflightService> GetPreflightService() { return this->PreflightService; }
 
 	bool bIsUserInAllowlist = false;
 
diff --git a/Engine/Source/Programs/SubmitTool/Private/View/Widgets/SJiraWidget.cpp b/Engine/Source/Programs/SubmitTool/Private/View/Widgets/SJiraWidget.cpp
--- a/Engine/Source/Programs/SubmitTool/Private/View/Widgets/SJiraWidget.cpp
+++ b/Engine/Source/Programs/SubmitTool/Private/View/Widgets/SJiraWidget.cpp
@@ -245,6 +245,44 @@ TSharedPtr<SWidget> SJiraWidget::BuildIssuesWidget()
 				.IsChecked_Lambda([this]() { return bIncludeClosedJira ? ECheckBoxState::Checked : ECheckBoxState::Unchecked; })
 				.OnCheckStateChanged_Lambda([this](ECheckBoxState newState) { bIncludeClosedJira = static_cast<bool>(newState); ApplyFilter(FilterText); })
 			]
+			+ SHorizontalBox::Slot()
+			.AutoWidth()			
+			.Padding(5,0,0,0)
+			.VAlign(VAlign_Center)
+			[			
+				SNew(SBox)
+				.MinDesiredWidth(110)
+				[
+					SNew(SButton)
+					.HAlign(EHorizontalAlignment::HAlign_Center)
+					.IsEnabled_Lambda([this]{ return FModelInterface::GetInputEnabled() && !JiraService.Pin()->bOngoingRequest; })
+					.Text_Lambda([this]{ 
+
+						FText RetText;
+						if (JiraService.Pin()->bOngoingRequest)
+						{
+							TStringBuilder<16> Message;
+						
+							Message << TEXT("Refreshing");
+						
+							size_t currSec = FDateTime::UtcNow().GetSecond() % 4;
+							for (size_t i = 0; i < currSec; ++i)
+							{
+								Message << TEXT('.');
+							}
+
+							RetText = FText::FromString(Message.ToString());
+						}	
+						else
+						{
+							RetText = FText::FromString("Refresh");
+						}
+
+						return RetText;
+					})
+				.OnClicked_Lambda([this]{ JiraService.Pin()->FetchJiraTickets(true); return FReply::Handled(); })
+				]
+			]
 		]
 
 		+ SVerticalBox::Slot()
@@ -332,20 +370,20 @@ TSharedPtr<SWidget> SJiraWidget::BuildIssuesWidget()
 					SNew(SButton)
 					.IsEnabled_Lambda(FModelInterface::GetInputEnabled)
 					.Text_Lambda([&issue = SelectedIssue, &targetTag = TargetTag]() -> FText {
-					if (issue == nullptr || targetTag == nullptr)
-					{
-						return FText::FromString("Select issue");
-					}
-
-					FString IssueKey = issue->Key;
-					TArray<FString> CurrentValues = targetTag->GetValues();
-					if (!CurrentValues.ContainsByPredicate([IssueKey](const FString& Value) { return Value.Equals(IssueKey, ESearchCase::IgnoreCase); }))
-					{
-						return FText::FromString("Add issue");
-					}
-
-					return FText::FromString("Remove issue");
-							})
+						if (issue == nullptr || targetTag == nullptr)
+						{
+							return FText::FromString("Select issue");
+						}
+
+						FString IssueKey = issue->Key;
+						TArray<FString> CurrentValues = targetTag->GetValues();
+						if (!CurrentValues.ContainsByPredicate([IssueKey](const FString& Value) { return Value.Equals(IssueKey, ESearchCase::IgnoreCase); }))
+						{
+							return FText::FromString("Add issue");
+						}
+
+						return FText::FromString("Remove issue");
+					})
 					.OnClicked(this, &SJiraWidget::OnSelectIssueClicked)
 								.IsEnabled_Lambda([&issue = SelectedIssue]() -> bool {return issue != nullptr; })
 								.ButtonStyle(FSubmitToolStyle::Get(), "PrimaryButton")
@@ -618,7 +656,7 @@ void SJiraWidget::ApplyFilter(const FString& InText)
 			return false;
 		}
 
-		return FilterText.IsEmpty() ? true : Issue->Summary.Contains(FilterText);
+		return FilterText.IsEmpty() ? true : Issue->Summary.Contains(FilterText) || Issue->Key.Contains(FilterText);
 	});
 	
 
diff --git a/Engine/Source/Programs/SubmitTool/Private/View/Widgets/SubmitToolWidget.cpp b/Engine/Source/Programs/SubmitTool/Private/View/Widgets/SubmitToolWidget.cpp
--- a/Engine/Source/Programs/SubmitTool/Private/View/Widgets/SubmitToolWidget.cpp
+++ b/Engine/Source/Programs/SubmitTool/Private/View/Widgets/SubmitToolWidget.cpp
@@ -667,7 +667,6 @@ TSharedRef<SWidget> SubmitToolWidget::BuildOutputLogWidget()
 	SummaryLogCreationParams.DefaultCategorySelection.Emplace(LogOutputDevice.GetCategoryName(), true);
 	SummaryLogCreationParams.DefaultCategorySelection.Emplace(LogSubmitToolP4.GetCategoryName(), true);
 	SummaryLogCreationParams.DefaultCategorySelection.Emplace(LogPresubmit.GetCategoryName(), false);
-	SummaryLogCreationParams.DefaultCategorySelection.Emplace("SourceControl", true);
 
 	SummaryLogDockTab = SNew(SDockTab)
 		.CanEverClose(false)
