Added user defined return type to source control services to make it more convenient to use source control calls

diff --git a/Engine/Source/Programs/SubmitTool/Private/Logic/ChangelistService.cpp b/Engine/Source/Programs/SubmitTool/Private/Logic/ChangelistService.cpp
--- a/Engine/Source/Programs/SubmitTool/Private/Logic/ChangelistService.cpp
+++ b/Engine/Source/Programs/SubmitTool/Private/Logic/ChangelistService.cpp
@@ -699,7 +699,7 @@ const TArray<FString>& FChangelistService::GetFilesDepotPaths(bool bForce)
 	if(bForce || (FilesDepotPaths.IsEmpty() && !FilesInCL.IsEmpty()))
 	{
 		FilesDepotPaths.Reset();
-		UE::Tasks::TTask<bool> OpenedTask = SourceControlService->RunCommand(TEXT("opened"), { "-c", CLID }, FOnSCCCommandComplete::CreateLambda(
+		UE::Tasks::TTask<FSCCResultNoRet> OpenedTask = SourceControlService->RunCommand(TEXT("opened"), { "-c", CLID }, FOnSCCCommandCompleteNoRet::CreateLambda(
 			[this](bool bSuccess, const FSCCRecordSet& InResultValues, const FSourceControlResultInfo& InResultsInfo)
 			{
 				if(bSuccess)
@@ -743,7 +743,7 @@ const TArray<FString>& FChangelistService::GetShelvedFilesDepotPaths(bool bForce
 	if(bForce || (ShelvedFilesDepotPaths.IsEmpty() && !ShelvedFilesInCL.IsEmpty()))
 	{
 		ShelvedFilesDepotPaths.Reset();
-		UE::Tasks::TTask<bool> DescribeTask = SourceControlService->RunCommand(TEXT("describe"), { "-S", CLID }, FOnSCCCommandComplete::CreateLambda(
+		UE::Tasks::TTask<FSCCResultNoRet> DescribeTask = SourceControlService->RunCommand(TEXT("describe"), { "-S", CLID }, FOnSCCCommandCompleteNoRet::CreateLambda(
 			[this](bool bSuccess, const FSCCRecordSet& InResultValues, const FSourceControlResultInfo& InResultsInfo)
 			{
 				if(bSuccess)
diff --git a/Engine/Source/Programs/SubmitTool/Private/Logic/P4LockdownService.cpp b/Engine/Source/Programs/SubmitTool/Private/Logic/P4LockdownService.cpp
--- a/Engine/Source/Programs/SubmitTool/Private/Logic/P4LockdownService.cpp
+++ b/Engine/Source/Programs/SubmitTool/Private/Logic/P4LockdownService.cpp
@@ -172,7 +172,7 @@ FP4LockdownService::FPathLockdownResult FP4LockdownService::IsPathInLockdown(con
 bool FP4LockdownService::GetConfigFile(const FString& InConfigId, const FString& InDepotPath)
 {
 	TArray<FSharedBuffer> FileBuffers;
-	if(ServiceProvider.Pin()->GetService<ISTSourceControlService>()->DownloadFiles(InDepotPath, FileBuffers).GetResult() && FileBuffers.Num() > 0)
+	if(ServiceProvider.Pin()->GetService<ISTSourceControlService>()->DownloadFiles(InDepotPath, FileBuffers).GetResult().bRequestSucceed && FileBuffers.Num() > 0)
 	{
 		FScopeLock Scope(&FileMutex);
 		
diff --git a/Engine/Source/Programs/SubmitTool/Private/Logic/Services/Interfaces/ISTSourceControlService.h b/Engine/Source/Programs/SubmitTool/Private/Logic/Services/Interfaces/ISTSourceControlService.h
--- a/Engine/Source/Programs/SubmitTool/Private/Logic/Services/Interfaces/ISTSourceControlService.h
+++ b/Engine/Source/Programs/SubmitTool/Private/Logic/Services/Interfaces/ISTSourceControlService.h
@@ -75,12 +75,24 @@ public:
 	TArray<const FString> AdditionalImportPaths;
 };
 
+template<typename UserType>
+struct FSCCResultRetVal
+{
+	bool bRequestSucceed;
+	TSharedPtr<UserType> UserData;
+};
+
+using FSCCResultNoRet = FSCCResultRetVal<void>;
+
 using FSCCRecordSet = TArray<TMap<FString, FString>>;
 
 DECLARE_MULTICAST_DELEGATE_OneParam(FOnUsersGet, TArray<TSharedPtr<FUserData>>& /*Users*/)
 DECLARE_MULTICAST_DELEGATE_OneParam(FOnGroupsGet, TArray<TSharedPtr<FString>>& /*Groups*/)
 DECLARE_MULTICAST_DELEGATE_TwoParams(FOnUsersAndGroupsGet, TArray<TSharedPtr<FUserData>>& /*Users*/, TArray<TSharedPtr<FString>>& /*Groups*/)
-DECLARE_DELEGATE_ThreeParams(FOnSCCCommandComplete, bool /*bSuccess*/, const FSCCRecordSet& /*InResultValues*/, const FSourceControlResultInfo& /*InResultsInfo*/)
+
+template<typename Type>
+using FOnSCCCommandCompleteRetVal = TDelegate<Type(bool /*bSuccess*/, const FSCCRecordSet& /*InResultValues*/, const FSourceControlResultInfo& /*InResultsInfo*/)>;
+using FOnSCCCommandCompleteNoRet = FOnSCCCommandCompleteRetVal<void>;
 
 class ISTSourceControlService : public ISubmitToolService
 {
@@ -90,7 +102,7 @@ public:
 	virtual void GetGroups(const FOnGroupsGet::FDelegate& Callback) = 0;
 	virtual void GetUsersAndGroups(const FOnUsersAndGroupsGet::FDelegate& Callback) = 0;
 
-	virtual UE::Tasks::TTask<bool> DownloadFiles(const FString& InFilepath, TArray<FSharedBuffer>& OutFileBuffers) = 0;
+	virtual UE::Tasks::TTask<FSCCResultNoRet> DownloadFiles(const FString& InFilepath, TArray<FSharedBuffer>& OutFileBuffers) = 0;
 
 	virtual bool IsAvailable() const = 0;
 
@@ -101,8 +113,6 @@ public:
 	virtual const TArray<TSharedPtr<FString>>& GetRecentGroups() const = 0;
 	virtual void AddRecentGroup(TSharedPtr<FString>& Group) = 0;
 
-	virtual UE::Tasks::TTask<bool> RunCommand(const FString& InCommand, const TArray<FString>& InAdditionalArgs, FOnSCCCommandComplete InCompleteCallback = nullptr, TArray<FSharedBuffer>* OutData = nullptr) = 0;
-
 	virtual TSharedPtr<FUserData> GetUserDataFromCache(const FString& Username) const = 0;
 
 	virtual const TArray<FSCCStream*>& GetClientStreams() const = 0;
@@ -111,6 +121,43 @@ public:
 	virtual const FString GetCurrentStreamName() = 0;
 	virtual const size_t GetDepotStreamLength(const FString& InDepotName) = 0;
 	virtual const FAuthTicket& GetAuthTicket() = 0;
+
+
+	template<typename T>
+	UE::Tasks::TTask<FSCCResultRetVal<T>> RunCommand(const FString& InCommand, const TArray<FString>& InAdditionalArgs, FOnSCCCommandCompleteRetVal<T> InCompleteCallback, TArray<FSharedBuffer>* OutData = nullptr)
+	{
+		return UE::Tasks::Launch(UE_SOURCE_LOCATION, [this, InCommand, InAdditionalArgs, InCompleteCallback, OutData] {
+			FSCCRecordSet ResultValues;
+			FSourceControlResultInfo ResultInfo;
+
+			bool bResult = RunCommandInternal(InCommand, InAdditionalArgs, ResultValues, ResultInfo, OutData);
+
+			TSharedPtr<T> UserTypeResult;
+
+			if (InCompleteCallback.IsBound())
+			{
+				if constexpr (!std::is_same<T, void>::value)
+				{
+					UserTypeResult = MakeShared<T>(InCompleteCallback.Execute(bResult, ResultValues, ResultInfo));
+				}
+				else
+				{
+					InCompleteCallback.Execute(bResult, ResultValues, ResultInfo);
+				}
+			}
+
+			return FSCCResultRetVal<T>(bResult, MoveTemp(UserTypeResult));
+			});
+	};
+
+	UE::Tasks::TTask<FSCCResultNoRet> RunCommand(const FString& InCommand, const TArray<FString>& InAdditionalArgs, TArray<FSharedBuffer>* OutData = nullptr)
+	{
+		return RunCommand<void>(InCommand, InAdditionalArgs, nullptr, OutData);
+	}
+
+protected:
+	virtual bool RunCommandInternal(const FString& InCommand, const TArray<FString>& InAdditionalArgs, FSCCRecordSet& OutResultValues, FSourceControlResultInfo& OutResultInfo, TArray<FSharedBuffer>* OutData = nullptr) = 0;
+
 };
 
 Expose_TNameOf(ISTSourceControlService);
\ No newline at end of file
diff --git a/Engine/Source/Programs/SubmitTool/Private/Logic/Services/SourceControl/SubmitToolPerforce.cpp b/Engine/Source/Programs/SubmitTool/Private/Logic/Services/SourceControl/SubmitToolPerforce.cpp
--- a/Engine/Source/Programs/SubmitTool/Private/Logic/Services/SourceControl/SubmitToolPerforce.cpp
+++ b/Engine/Source/Programs/SubmitTool/Private/Logic/Services/SourceControl/SubmitToolPerforce.cpp
@@ -172,7 +172,7 @@ const TUniquePtr<ISourceControlProvider>& FSubmitToolPerforce::GetProvider() con
 	return SCCProvider;
 }
 
-UE::Tasks::TTask<bool> FSubmitToolPerforce::GetStream(const FString& InStream, bool bRequestHierarchy)
+UE::Tasks::TTask<FSCCResultNoRet> FSubmitToolPerforce::GetStream(const FString& InStream, bool bRequestHierarchy)
 {
 	TArray<FString> Args = { TEXT("-o") };
 
@@ -181,7 +181,7 @@ UE::Tasks::TTask<bool> FSubmitToolPerforce::GetStream(const FString& InStream, b
 		Args.Add(InStream);
 	}
 
-	return RunCommand(TEXT("stream"), Args, FOnSCCCommandComplete::CreateLambda(
+	return RunCommand(TEXT("stream"), Args, FOnSCCCommandCompleteNoRet::CreateLambda(
 		[this, bRequestHierarchy, InStream](bool bSuccess, const FSCCRecordSet& InResultValues, const FSourceControlResultInfo& InResultsInfo)
 		{
 			FScopeLock StreamLock(&StreamMutex);
@@ -264,7 +264,7 @@ void FSubmitToolPerforce::GetUsers(const FOnUsersGet::FDelegate& Callback)
 
 	if(!UserTask.IsValid())
 	{
-		UserTask = RunCommand(TEXT("users"), {}, FOnSCCCommandComplete::CreateLambda(
+		UserTask = RunCommand(TEXT("users"), {}, FOnSCCCommandCompleteNoRet::CreateLambda(
 			[this](bool bSuccess, const FSCCRecordSet& InResultValues, const FSourceControlResultInfo& InResultsInfo)
 			{
 				for(const TMap<FString, FString>& Record : InResultValues)
@@ -710,9 +710,9 @@ const FString FSubmitToolPerforce::GetRecentGroupsFilepath() const
 	return FPaths::Combine(FSubmitToolUtils::GetLocalAppDataPath(), TEXT("SubmitTool"), TEXT("recent_groups.dat"));
 }
 
-UE::Tasks::TTask<bool> FSubmitToolPerforce::DownloadFiles(const FString& InFilepath, TArray<FSharedBuffer>& OutFileBuffers)
+UE::Tasks::TTask<FSCCResultNoRet> FSubmitToolPerforce::DownloadFiles(const FString& InFilepath, TArray<FSharedBuffer>& OutFileBuffers)
 {
-	return RunCommand(TEXT("print"), { TEXT("-q"), InFilepath }, nullptr, &OutFileBuffers);
+	return RunCommand(TEXT("print"), { TEXT("-q"), InFilepath }, &OutFileBuffers);
 }
 
 void FSubmitToolPerforce::GetUsersAndGroups(const FOnUsersAndGroupsGet::FDelegate& Callback)
@@ -732,21 +732,12 @@ void FSubmitToolPerforce::GetUsersAndGroups(const FOnUsersAndGroupsGet::FDelegat
 	GetGroups(FOnGroupsGet::FDelegate::CreateLambda([OnUsersAndGroupsReady, this](TArray<TSharedPtr<FString>>&) { bGroupsReady = true; OnUsersAndGroupsReady(); }));
 }
 
-UE::Tasks::TTask<bool> FSubmitToolPerforce::RunCommand(const FString& InCommand, const TArray<FString>& InAdditionalArgs, FOnSCCCommandComplete InCompleteCallback, TArray<FSharedBuffer>* OutData)
+bool FSubmitToolPerforce::RunCommandInternal(const FString& InCommand, const TArray<FString>& InAdditionalArgs, FSCCRecordSet& OutResultValues, FSourceControlResultInfo& OutResultInfo, TArray<FSharedBuffer>* OutData)
 {
-	return UE::Tasks::Launch(UE_SOURCE_LOCATION, [this, InCommand, InAdditionalArgs, InCompleteCallback, OutData] {
-		FSCCRecordSet ResultValues;
-		FSourceControlResultInfo ResultInfo;
-
-		EP4ClientUserFlags Flags = DefaultFlags;
-		Flags |= bIsUnicodeServer ? EP4ClientUserFlags::UnicodeServer : EP4ClientUserFlags::None;
-
-		bool bResult = RunP4Command(InCommand, InAdditionalArgs, ResultValues, ResultInfo, Flags, OutData);
-
-		InCompleteCallback.ExecuteIfBound(bResult, ResultValues, ResultInfo);
-
-		return bResult;
-		});
+	EP4ClientUserFlags Flags = DefaultFlags;
+	Flags |= bIsUnicodeServer ? EP4ClientUserFlags::UnicodeServer : EP4ClientUserFlags::None;
+	
+	return RunP4Command(InCommand, InAdditionalArgs, OutResultValues, OutResultInfo, Flags, OutData);
 }
 
 const TArray<FSCCStream*>& FSubmitToolPerforce::GetClientStreams() const
@@ -764,7 +755,7 @@ const FSCCStream* FSubmitToolPerforce::GetSCCStream(const FString& InStreamName)
 	}
 	else
 	{
-		if(GetStream(InStreamName, false).GetResult() && Streams.Contains(InStreamName))
+		if(GetStream(InStreamName, false).GetResult().bRequestSucceed && Streams.Contains(InStreamName))
 		{
 			FScopeLock Stream(&StreamMutex);
 			return Streams[InStreamName].Get();
@@ -812,7 +803,7 @@ const size_t FSubmitToolPerforce::GetDepotStreamLength(const FString& InDepotNam
 {
 	if(!DepotStreamLengths.Contains(InDepotName))
 	{
-		RunCommand(TEXT("depot"), { TEXT("-o"), InDepotName }, FOnSCCCommandComplete::CreateLambda(
+		RunCommand(TEXT("depot"), { TEXT("-o"), InDepotName }, FOnSCCCommandCompleteNoRet::CreateLambda(
 			[this, InDepotName](bool bSuccess, const FSCCRecordSet& InResultValues, const FSourceControlResultInfo& InResultsInfo)
 			{
 				if(bSuccess)
diff --git a/Engine/Source/Programs/SubmitTool/Private/Logic/Services/SourceControl/SubmitToolPerforce.h b/Engine/Source/Programs/SubmitTool/Private/Logic/Services/SourceControl/SubmitToolPerforce.h
--- a/Engine/Source/Programs/SubmitTool/Private/Logic/Services/SourceControl/SubmitToolPerforce.h
+++ b/Engine/Source/Programs/SubmitTool/Private/Logic/Services/SourceControl/SubmitToolPerforce.h
@@ -82,7 +82,7 @@ public:
 	virtual void GetGroups(const FOnGroupsGet::FDelegate& Callback) override;
 	virtual void GetUsersAndGroups(const FOnUsersAndGroupsGet::FDelegate& Callback) override;
 
-	virtual UE::Tasks::TTask<bool> DownloadFiles(const FString& InFilepath, TArray<FSharedBuffer>& OutFileBuffers) override;
+	virtual UE::Tasks::TTask<FSCCResultNoRet> DownloadFiles(const FString& InFilepath, TArray<FSharedBuffer>& OutFileBuffers) override;
 
 	virtual bool IsAvailable() const override
 	{
@@ -96,8 +96,6 @@ public:
 	virtual const TArray<TSharedPtr<FString>>& GetRecentGroups() const override;
 	virtual void AddRecentGroup(TSharedPtr<FString>& Group) override;
 
-	virtual UE::Tasks::TTask<bool> RunCommand(const FString& InCommand, const TArray<FString>& InAdditionalArgs, FOnSCCCommandComplete InCompleteCallback = nullptr, TArray<FSharedBuffer>* OutData = nullptr) override;
-
 	virtual TSharedPtr<FUserData> GetUserDataFromCache(const FString& Username) const override;
 
 	virtual const TArray<FSCCStream*>& GetClientStreams() const override;
@@ -106,16 +104,19 @@ public:
 	virtual const FString GetCurrentStreamName() override;
 	virtual const size_t GetDepotStreamLength(const FString& InDepotName) override;
 	virtual const FAuthTicket& GetAuthTicket() override;
-								  
+			
+protected:
+	virtual bool RunCommandInternal(const FString& InCommand, const TArray<FString>& InAdditionalArgs, FSCCRecordSet& OutResultValues, FSourceControlResultInfo& OutResultInfo, TArray<FSharedBuffer>* OutData = nullptr) override;
+
 private:
 	bool RunP4Command(const FString& InCommand, const TArray<FString>& InAdditionalArgs, TArray<TMap<FString, FString>>& OutResultValues, FSourceControlResultInfo& OutResults, EP4ClientUserFlags InCmdFlags, TArray<FSharedBuffer>* OutData = nullptr);
 
 	const FSubmitToolParameters& Parameters;
-	UE::Tasks::TTask<bool> GetStream(const FString& InStream = FString(), bool bRequestHierarchy = true);
+	UE::Tasks::TTask<FSCCResultNoRet> GetStream(const FString& InStream = FString(), bool bRequestHierarchy = true);
 	TMap<FString, TSharedPtr<FUserData>> CachedUsers;
 	TArray<TSharedPtr<FUserData>> CachedUsersArray;
 	TArray<TSharedPtr<FString>> CachedGroupsArray;
-	UE::Tasks::TTask<bool> UserTask;
+	UE::Tasks::TTask<FSCCResultNoRet> UserTask;
 	UE::Tasks::TTask<bool> GroupTask;
 	bool bUsersReady = false;
 	bool bGroupsReady = false;
diff --git a/Engine/Source/Programs/SubmitTool/Private/Logic/SwarmService.cpp b/Engine/Source/Programs/SubmitTool/Private/Logic/SwarmService.cpp
--- a/Engine/Source/Programs/SubmitTool/Private/Logic/SwarmService.cpp
+++ b/Engine/Source/Programs/SubmitTool/Private/Logic/SwarmService.cpp
@@ -297,7 +297,7 @@ const FString& FSwarmService::GetSwarmURL()
 {
 	if(SwarmURL.IsEmpty())
 	{
-		ServiceProvider.Pin()->GetService<ISTSourceControlService>()->RunCommand(TEXT("property"), { "-l", "-n", "P4.Swarm.URL" }, FOnSCCCommandComplete::CreateLambda(
+		ServiceProvider.Pin()->GetService<ISTSourceControlService>()->RunCommand(TEXT("property"), { "-l", "-n", "P4.Swarm.URL" }, FOnSCCCommandCompleteNoRet::CreateLambda(
 			[this](bool bSuccess, const FSCCRecordSet& InResultValues, const FSourceControlResultInfo& InResultsInfo)
 			{
 				if(bSuccess && InResultValues.Num() > 0 && InResultValues[0].Contains(TEXT("value")))
