Added feature to disable validators manually. This will add a tag to the CL specifying which validators have been disabled
Added optimizations to Applies to File Validator checks

diff --git a/Engine/Source/Programs/SubmitTool/Private/Logic/TasksService.cpp b/Engine/Source/Programs/SubmitTool/Private/Logic/TasksService.cpp
--- a/Engine/Source/Programs/SubmitTool/Private/Logic/TasksService.cpp
+++ b/Engine/Source/Programs/SubmitTool/Private/Logic/TasksService.cpp
@@ -119,9 +119,9 @@ void FTasksService::QueueByArea(const ETaskArea& InArea)
 bool FTasksService::QueueForExecution(const TSharedPtr<FValidatorBase>& InTask, bool InbForceRun, TSet<FName>& InOutVisitedTasks)
 {
 	InOutVisitedTasks.Add(InTask->GetValidatorNameId());
-	if(InTask->GetIsQueued())
+	if(InTask->GetIsQueued() || InTask->Definition->bIsDisabled)
 	{
-		// If already queued
+		// If already queued or disabled
 		return false;
 	}
 
@@ -237,6 +237,14 @@ void FTasksService::CheckForTagSkips()
 	}
 }
 
+void FTasksService::ToggleEnabled(const FName& InTaskId)
+{
+	if (Tasks.Contains(InTaskId))
+	{
+		Tasks[InTaskId]->ToggleEnabled();
+	}
+}
+
 bool FTasksService::GetIsAnyTaskRunning() const
 {
 	bool bAnyTaskRunningOrQueued = false;
@@ -492,3 +500,15 @@ void FTasksService::StopTasksByArea(const ETaskArea& InArea)
 		}
 	}
 }
+
+const TWeakPtr<const FValidatorBase> FTasksService::GetTask(const FName& InId) const
+{
+	if (Tasks.Contains(InId))
+	{
+		return Tasks[InId];
+	}
+	else
+	{
+		return nullptr;
+	}
+}
diff --git a/Engine/Source/Programs/SubmitTool/Private/Logic/TasksService.h b/Engine/Source/Programs/SubmitTool/Private/Logic/TasksService.h
--- a/Engine/Source/Programs/SubmitTool/Private/Logic/TasksService.h
+++ b/Engine/Source/Programs/SubmitTool/Private/Logic/TasksService.h
@@ -31,11 +31,13 @@ public:
 	virtual void InvalidateForChanges(ETaskArea InChangeType);
 	virtual void CheckForLocalFileEdit();
 	virtual void CheckForTagSkips();
+	virtual void ToggleEnabled(const FName& InTaskId);
 
 	virtual bool GetIsRunSuccessful(bool bWaitForOptionalCompletes = true) const;
 	virtual bool GetIsAnyTaskRunning() const;
 	virtual bool AreTasksPendingQueue() const;
 
+	virtual const TWeakPtr<const FValidatorBase> GetTask(const FName& InId) const;
 	virtual const TArray<TWeakPtr<const FValidatorBase>>& GetTasks() const;
 	virtual const TArray<TWeakPtr<const FValidatorBase>> GetTasksOfType(const FString& TaskType) const;
 
diff --git a/Engine/Source/Programs/SubmitTool/Private/Logic/Validators/ValidatorBase.cpp b/Engine/Source/Programs/SubmitTool/Private/Logic/Validators/ValidatorBase.cpp
--- a/Engine/Source/Programs/SubmitTool/Private/Logic/Validators/ValidatorBase.cpp
+++ b/Engine/Source/Programs/SubmitTool/Private/Logic/Validators/ValidatorBase.cpp
@@ -51,6 +51,12 @@ FValidatorBase::~FValidatorBase()
 
 void FValidatorBase::StartValidation()
 {
+	if (Definition->bIsDisabled)
+	{
+		State = EValidationStates::Disabled;
+		return;
+	}
+
 	RunTime = 0;
 	Start = FDateTime::UtcNow();
 	State = EValidationStates::Running;
@@ -128,6 +134,22 @@ void FValidatorBase::StartValidation()
 	}
 }
 
+void FValidatorBase::ToggleEnabled()
+{
+	FValidatorDefinition* ModifiableDefinition = const_cast<FValidatorDefinition*>(Definition.Get());
+	ModifiableDefinition->bIsDisabled = !ModifiableDefinition->bIsDisabled;
+
+	CancelValidation();
+
+	if(ModifiableDefinition->bIsDisabled)
+	{
+		State = EValidationStates::Disabled;
+	}
+	else
+	{
+		State = EValidationStates::Not_Run;
+	}
+}
 
 void FValidatorBase::Tick(float InDeltaTime)
 {
@@ -381,12 +403,19 @@ bool FValidatorBase::AppliesToFile(const FSourceControlStateRef InFile, bool Inb
 		return true;
 	}
 
+	FString Filename = InFile->GetFilename();
+	FPaths::NormalizeFilename(Filename);
+	uint32 FileHash = GetTypeHash(InFile->GetFilename());
+
+	if (FileHashes.Contains(FileHash))
+	{
+		bIncluded = FileHashes[FileHash];
+	}
+	else
+	{
 	if(!InFile->IsDeleted()
 		|| (InFile->IsDeleted() && Definition->bAcceptDeletedFiles))
 	{
-		FString Filename = InFile->GetFilename();
-		FPaths::NormalizeFilename(Filename);
-
 		if (!Definition->AppliesToCLRegex.IsEmpty())
 		{
 			const FString RegexPat = FConfiguration::Substitute(Definition->AppliesToCLRegex);
@@ -461,11 +490,15 @@ bool FValidatorBase::AppliesToFile(const FSourceControlStateRef InFile, bool Inb
 				}
 			}
 		}
+		}
+
+		FileHashes.FindOrAdd(FileHash, bIncluded);
+	}
 
 		if(InbAllowIncremental && bIncluded)
 		{
 			FFileManagerGeneric FileManager;
-			FDateTime LastValidation = ServiceProvider.Pin()->GetService<ICacheDataService>()->GetLastValidationDate(ServiceProvider.Pin()->GetService<FChangelistService>()->GetCLID(), GetValidatorNameId(), GetValidationConfigId(), InFile->GetFilename());
+		FDateTime LastValidation = ServiceProvider.Pin()->GetService<ICacheDataService>()->GetLastValidationDate(ServiceProvider.Pin()->GetService<FChangelistService>()->GetCLID(), GetValidatorNameId(), GetValidationConfigId(), Filename);
 			FFileStatData FileModifiedDate = FileManager.GetStatData(*Filename);
 			if(LastValidation != FDateTime::MinValue() && FileModifiedDate.ModificationTime < LastValidation)
 			{
@@ -473,7 +506,6 @@ bool FValidatorBase::AppliesToFile(const FSourceControlStateRef InFile, bool Inb
 				bIncluded = false;
 			}
 		}
-	}
 
 	return bIncluded;
 }
diff --git a/Engine/Source/Programs/SubmitTool/Private/Logic/Validators/ValidatorBase.h b/Engine/Source/Programs/SubmitTool/Private/Logic/Validators/ValidatorBase.h
--- a/Engine/Source/Programs/SubmitTool/Private/Logic/Validators/ValidatorBase.h
+++ b/Engine/Source/Programs/SubmitTool/Private/Logic/Validators/ValidatorBase.h
@@ -45,7 +45,8 @@ enum class EValidationStates : uint8
 	Valid,
 	Failed,
 	Timeout,
-	Queued
+	Queued,
+	Disabled
 };
 
 template<class T>
@@ -78,8 +79,15 @@ public:
 
 	virtual void StartValidation();
 
+	virtual void ToggleEnabled();
+
 	void Invalidate(bool bForce = false)
 	{
+		if (State == EValidationStates::Disabled)
+		{
+			return;
+		}
+
 		if(State == EValidationStates::Running)
 		{
 			CancelValidation();
@@ -92,7 +100,7 @@ public:
 
 	bool GetHasPassed() const
 	{
-		return State == EValidationStates::Valid || State == EValidationStates::Skipped || State == EValidationStates::Not_Applicable;
+		return State == EValidationStates::Valid || State == EValidationStates::Skipped || State == EValidationStates::Not_Applicable || State == EValidationStates::Disabled;
 	}
 
 	bool GetIsRunningOrQueued() const
@@ -124,6 +132,11 @@ public:
 
 	virtual void CancelValidation(bool InbAsFailed = false)
 	{
+		if (State == EValidationStates::Disabled)
+		{
+			return;
+		}
+
 		if(State == EValidationStates::Running)
 		{
 			StopInternalValidations();
@@ -248,6 +261,8 @@ protected:
 private:
 	EValidationStates State = EValidationStates::Not_Run;
 
+
+	mutable TMap<uint32, bool> FileHashes;
 	mutable FCriticalSection Mutex;
 };
 
diff --git a/Engine/Source/Programs/SubmitTool/Private/Logic/Validators/ValidatorDefinition.h b/Engine/Source/Programs/SubmitTool/Private/Logic/Validators/ValidatorDefinition.h
--- a/Engine/Source/Programs/SubmitTool/Private/Logic/Validators/ValidatorDefinition.h
+++ b/Engine/Source/Programs/SubmitTool/Private/Logic/Validators/ValidatorDefinition.h
@@ -50,11 +50,17 @@ struct FValidatorDefinition
 	FString Type;
 
 	///
-	///	If true, this validator will be disabled and not loaded, useful for config overriding a parent validator
+	///	If true, this validator will be disabled
 	/// 
 	UPROPERTY()
 	bool bIsDisabled = false;
 
+	///
+	///	If false, this validator won't have the disable button available
+	/// 
+	UPROPERTY()
+	bool bCanBeDisabledByUser = true;
+
 	/// 
 	/// Whether the validator is required to allow submission or not, a failing required validation will always block submission
 	/// 
diff --git a/Engine/Source/Programs/SubmitTool/Private/Models/ModelInterface.cpp b/Engine/Source/Programs/SubmitTool/Private/Models/ModelInterface.cpp
--- a/Engine/Source/Programs/SubmitTool/Private/Models/ModelInterface.cpp
+++ b/Engine/Source/Programs/SubmitTool/Private/Models/ModelInterface.cpp
@@ -353,6 +353,67 @@ bool FModelInterface::CanLaunchPreflight() const
 	return true;
 }
 
+void FModelInterface::EvaluateDisabledValidatorsTag()
+{
+	bool bUpdateDescription = false;
+
+	FString DescriptionCopy = ChangelistService->GetCLDescription();
+	const FString DisabledTag = TEXT("#DisabledValidations ");
+
+	int32 Loc = DescriptionCopy.Find(DisabledTag);
+	if (Loc != INDEX_NONE)
+	{
+		size_t EndPos = Loc + DisabledTag.Len();
+		while (EndPos < DescriptionCopy.Len())
+		{
+			if (DescriptionCopy[EndPos] == TCHAR('\n'))
+			{
+				++EndPos;
+				break;
+			}
+			++EndPos;
+		}
+
+		if (Loc != 0 && DescriptionCopy[Loc - 1] == '\n' && (EndPos - DescriptionCopy.Len()) < 2)
+		{
+			Loc--;
+		}
+
+		UE_LOG(LogSubmitToolDebug, Log, TEXT("Removed Disabled Validations tag %s"), *DescriptionCopy.Mid(Loc, EndPos - Loc));
+		DescriptionCopy.RemoveAt(Loc, EndPos - Loc);
+		bUpdateDescription = true;
+	}
+
+
+	TStringBuilder<256> DisabledValidatorsString;
+	for (const TWeakPtr<const FValidatorBase> Validator : ValidationService->GetTasks())
+	{
+		if (Validator.Pin()->Definition->bIsDisabled)
+		{
+			if (DisabledValidatorsString.Len() != 0)
+			{
+				DisabledValidatorsString << TEXT(", ");
+			}
+
+			DisabledValidatorsString << Validator.Pin()->GetValidatorNameId();
+		}
+	}
+
+	if (DisabledValidatorsString.Len() != 0)
+	{
+		DisabledValidatorsString.InsertAt(0, DisabledTag);
+		const FString FinalString = DisabledValidatorsString.ToString();
+		UE_LOG(LogSubmitToolDebug, Log, TEXT("Added Disabled validators tag %s"), *FinalString);
+		DescriptionCopy.Append(TEXT("\n") + FinalString);
+	}
+
+	if (Loc != INDEX_NONE || DisabledValidatorsString.Len() != 0)
+	{
+		ChangelistService->SetCLDescription(DescriptionCopy, true);
+		TagService->ParseCLDescription();
+	}
+}
+
 void FModelInterface::ReevaluateSubmitToolTag()
 {
 	UpdateSubmitToolTag(ValidationService->GetIsRunSuccessful(!IsIntegrationRequired()));
@@ -361,7 +422,7 @@ void FModelInterface::ReevaluateSubmitToolTag()
 void FModelInterface::UpdateSubmitToolTag(bool InbAdd)
 {
 	// add a special tag to the CL description
-	FString SubmitToolTag = FString::Format(TEXT("#submittool {0}\n"), { FAppVersion::GetVersion() });
+	FString SubmitToolTag = FString::Format(TEXT("#submittool {0}"), { FAppVersion::GetVersion() });
 	FString DescriptionCopy = ChangelistService->GetCLDescription();
 
 	if (InbAdd)
@@ -489,6 +550,7 @@ void FModelInterface::StartSubmitProcess(bool bSkipShelfDialog)
 			PrepareSubmitCallBack.Broadcast();
 		}
 
+		EvaluateDisabledValidatorsTag();
 		UpdateSubmitToolTag(true);
 
 		if (Parameters.IncompatibleFilesParams.IncompatibleFileGroups.Num() > 0)
@@ -1043,3 +1105,9 @@ void FModelInterface::OnSwarmCreateCompleted(bool InResult, const FString& InErr
 		ShowSwarmReview();
 	}
 }
+
+void FModelInterface::ToggleValidator(const FName& ValidatorId)
+{
+	ValidationService->ToggleEnabled(ValidatorId);
+	EvaluateDisabledValidatorsTag();	
+}
diff --git a/Engine/Source/Programs/SubmitTool/Private/Models/ModelInterface.h b/Engine/Source/Programs/SubmitTool/Private/Models/ModelInterface.h
--- a/Engine/Source/Programs/SubmitTool/Private/Models/ModelInterface.h
+++ b/Engine/Source/Programs/SubmitTool/Private/Models/ModelInterface.h
@@ -128,9 +128,11 @@ public:
 
 	void ValidateChangelist() const													{ ValidationService->QueueAll(); }
 	void ValidateSingle(const FName& ValidatorId, bool bForce = true)				{ ValidationService->QueueSingle(ValidatorId, bForce); }
+	void ToggleValidator(const FName& ValidatorId);
 	void ValidateCLDescription() const												{ ValidationService->StopTasksByArea(ETaskArea::Changelist); ValidationService->QueueByArea(ETaskArea::Changelist); }
 	bool IsCLValid() const															{ return ValidationService->GetIsRunSuccessful(!IsIntegrationRequired()); }
 	bool CanLaunchPreflight() const;
+	void EvaluateDisabledValidatorsTag();
 	void ReevaluateSubmitToolTag();
 	void UpdateSubmitToolTag(bool InbAdd);
 	bool HasSubmitToolTag() const;
diff --git a/Engine/Source/Programs/SubmitTool/Private/View/SubmitToolStyle.cpp b/Engine/Source/Programs/SubmitTool/Private/View/SubmitToolStyle.cpp
--- a/Engine/Source/Programs/SubmitTool/Private/View/SubmitToolStyle.cpp
+++ b/Engine/Source/Programs/SubmitTool/Private/View/SubmitToolStyle.cpp
@@ -111,6 +111,7 @@ FSubmitToolStyle::FSubmitToolStyle() : FSlateStyleSet(StyleName)
 	Set("ValidatorStateWarning", FStyleColors::Warning.GetSpecifiedColor());
 	Set("ValidatorStateFail", FStyleColors::Error.GetSpecifiedColor());
 	Set("ValidatorStateSuccess", FStyleColors::Success.GetSpecifiedColor());
+	Set("ValidatorStateDisabled", FStyleColors::AccentGray.GetSpecifiedColor());
 	Set("ValidatorStateNormal", FStyleColors::AccentBlue.GetSpecifiedColor());
 	static const FSlateColor& TabFlashColor = MakeShared<FLinearColor>(COLOR("18A0FBFF"));
 	FLinearColor DockColor_Active(FColor(62, 62, 62));
diff --git a/Engine/Source/Programs/SubmitTool/Private/View/Widgets/ValidatorsWidget.cpp b/Engine/Source/Programs/SubmitTool/Private/View/Widgets/ValidatorsWidget.cpp
--- a/Engine/Source/Programs/SubmitTool/Private/View/Widgets/ValidatorsWidget.cpp
+++ b/Engine/Source/Programs/SubmitTool/Private/View/Widgets/ValidatorsWidget.cpp
@@ -30,9 +30,9 @@ void SValidatorsWidget::Construct(const FArguments& InArgs)
 {
 	Columns = {
 		FValidatorColumn(ValidatorsColumns::Name, 1, true, [](TWeakPtr<const FValidatorBase> In) {return In.Pin()->GetValidatorName(); }),
-		FValidatorColumn(ValidatorsColumns::Required, 100, false, [](TWeakPtr<const FValidatorBase> In) {return In.Pin()->Definition->IsRequired ? TEXT("Yes") : TEXT("No"); }),
+		FValidatorColumn(ValidatorsColumns::Required, 80, false, [](TWeakPtr<const FValidatorBase> In) {return In.Pin()->Definition->IsRequired ? TEXT("Yes") : TEXT("No"); }),
 		FValidatorColumn(ValidatorsColumns::Status, 100, false, [](TWeakPtr<const FValidatorBase> In) {return In.Pin()->GetStatusText(); }),
-		FValidatorColumn(ValidatorsColumns::Actions, 300, false, nullptr),
+		FValidatorColumn(ValidatorsColumns::Actions, 350, false, nullptr),
 	};
 
 	ModelInterface = InArgs._ModelInterface;
@@ -43,7 +43,7 @@ void SValidatorsWidget::Construct(const FArguments& InArgs)
 		BuildValidatorsView()
 	];
 
-	OnFilesRefreshed = ModelInterface->FileRefreshedCallback.AddLambda([this] { RefreshValidatorView(); });
+	OnFilesRefreshed = ModelInterface->FileRefreshedCallback.AddLambda([this] { ReevaluateValidatorSection(); });
 	OnPrepareSubmit = FModelInterface::OnStateChanged.AddLambda([this] (const ESubmitToolAppState InFrom, const ESubmitToolAppState InTo)
 		{ 
 			if (InTo == ESubmitToolAppState::Submitting || InFrom == ESubmitToolAppState::Submitting)
@@ -59,13 +59,40 @@ SValidatorsWidget::~SValidatorsWidget()
 	FModelInterface::OnStateChanged.Remove(OnPrepareSubmit);
 }
 
+void SValidatorsWidget::ReevaluateValidatorSection()
+{
+	for (const TPair<FName, TArray<TWeakPtr<const FValidatorBase>>>& Pair : Validators)
+	{
+		for (size_t i = 0;i < Pair.Value.Num();++i)
+		{
+			FName Section;
+			
+			if (!Pair.Value[i].Pin()->Definition->bIsDisabled && !Pair.Value[i].Pin()->IsRelevantToCL())
+			{
+				Section = InactiveSection;
+			}
+			else
+			{
+				Section = ActiveSection;
+			}
+
+			if (Section != Pair.Key)
+			{
+				Validators.FindOrAdd(Section).Add(Pair.Value[i]);
+				Validators[Pair.Key].RemoveAt(i);
+				--i;
+			}
+		}
+	}
+}
+
 TSharedRef<SWidget> SValidatorsWidget::BuildValidatorsView(bool bListPreSubmitOperations)
 {
 	Validators.Reset();
 
 	TArray<TWeakPtr<const FValidatorBase>> ValidatorsArray = bListPreSubmitOperations ? ModelInterface->GetPreSubmitOperations() : ModelInterface->GetValidators();
-	const FName ActiveSection = bListPreSubmitOperations ? FName(TEXT("Pre-Submit Operations")) : FName(TEXT("Active Validations"));
-	const FName InactiveSection = bListPreSubmitOperations ? FName(TEXT("Inactive Pre-Submit Operations (Not Applicable to your CL)")) : FName(TEXT("Inactive Validators (Not Applicable to your CL)"));
+	ActiveSection = bListPreSubmitOperations ? FName(TEXT("Pre-Submit Operations")) : FName(TEXT("Active Validations"));
+	InactiveSection = bListPreSubmitOperations ? FName(TEXT("Inactive Pre-Submit Operations (Not Applicable to your CL)")) : FName(TEXT("Inactive Validators (Not Applicable to your CL)"));
 
 	TSharedRef<SVerticalBox> VBox = SNew(SVerticalBox);
 	
@@ -73,7 +100,7 @@ TSharedRef<SWidget> SValidatorsWidget::BuildValidatorsView(bool bListPreSubmitOp
 	{
 		FName DisplayName;
 
-		if (!Validator.Pin()->IsRelevantToCL())
+		if (!Validator.Pin()->Definition->bIsDisabled && !Validator.Pin()->IsRelevantToCL())
 		{
 			DisplayName = InactiveSection;
 		}
@@ -103,7 +130,7 @@ TSharedRef<SWidget> SValidatorsWidget::BuildValidatorsView(bool bListPreSubmitOp
 				.BorderImage(FAppStyle::Get().GetBrush("Brushes.Header"))
 				.BodyBorderImage(FAppStyle::Get().GetBrush("Brushes.Recessed"))
 				.OnAreaExpansionChanged_Lambda([DisplayName = Pair.Key](bool bExpanded) { FSubmitToolUserPrefs::Get()->UISectionExpandState.FindOrAdd(DisplayName, false) = bExpanded; })
-				.BorderBackgroundColor_Lambda([&ValidatorArray = Pair.Value, DisplayName = Pair.Key, InactiveSection]
+				.BorderBackgroundColor_Lambda([&ValidatorArray = Pair.Value, &DisplayName = Pair.Key, this]
 				{
 					if (DisplayName.IsEqual(InactiveSection))
 					{
@@ -205,6 +232,7 @@ class SValidatorNode : public SMultiColumnTableRow<TWeakPtr<const FValidatorBase
 {
 public:
 	SLATE_BEGIN_ARGS(SValidatorNode) {}
+		SLATE_ARGUMENT(SValidatorsWidget*, ParentWidget)
 		SLATE_ARGUMENT(FModelInterface*, ModelInterface)
 		SLATE_ARGUMENT(bool, IsRelevantToCL)
 		SLATE_EVENT(FOnViewValidatorLog, ViewLogCallback)
@@ -213,6 +241,7 @@ public:
 	void Construct(const FArguments& InArgs, const TSharedRef<STableViewBase>& InOwnerTable, TWeakPtr<const FValidatorBase> InNode)
 	{
 		ModelInterface = InArgs._ModelInterface;
+		ParentWidget = InArgs._ParentWidget;
 		bIsRelevantToCL = InArgs._IsRelevantToCL;
 		Node = InNode;
 		ViewLogCallback = InArgs._ViewLogCallback;
@@ -345,6 +374,9 @@ public:
 					case EValidationStates::Running:
 						return FAppStyle::GetColor("ValidatorStateWarning");
 						break;
+					case EValidationStates::Disabled:
+						return FAppStyle::GetColor("ValidatorStateDisabled");
+						break;
 					case EValidationStates::Not_Run:
 					default:
 						return FAppStyle::GetColor("ValidatorStateNormal");
@@ -359,7 +391,7 @@ public:
 				.Padding(3,0,0,0)
 				[
 					SNew(SButton)
-					.IsEnabled_Lambda([this]() {return ModelInterface->GetInputEnabled() && bIsRelevantToCL; })
+					.IsEnabled_Lambda([this]() {return ModelInterface->GetInputEnabled() && Node.Pin()->GetValidatorState() != EValidationStates::Disabled && bIsRelevantToCL; })
 					.HAlign(EHorizontalAlignment::HAlign_Center)
 					.VAlign(EVerticalAlignment::VAlign_Center)
 					.OnClicked_Lambda([this]()
@@ -379,7 +411,7 @@ public:
 				.Padding(3,0,0,0)
 				[
 					SNew(SButton)
-					.IsEnabled_Lambda([this]() {return ModelInterface->GetInputEnabled() && bIsRelevantToCL; })
+					.IsEnabled_Lambda([this]() {return ModelInterface->GetInputEnabled() && Node.Pin()->GetValidatorState() != EValidationStates::Disabled && bIsRelevantToCL; })
 					.HAlign(EHorizontalAlignment::HAlign_Center)
 					.VAlign(EVerticalAlignment::VAlign_Center)
 					.OnClicked_Lambda([this]()
@@ -402,7 +434,7 @@ public:
 				.Padding(3, 0, 3, 0)
 				[
 					SNew(SButton)
-					.IsEnabled_Lambda([this]() { return Node.Pin()->GetValidatorState() != EValidationStates::Not_Run && Node.Pin()->GetValidatorState() != EValidationStates::Skipped && Node.Pin()->GetValidatorState() != EValidationStates::Not_Applicable; })
+					.IsEnabled_Lambda([this]() { return Node.Pin()->GetValidatorState() != EValidationStates::Not_Run && Node.Pin()->GetValidatorState() != EValidationStates::Skipped && Node.Pin()->GetValidatorState() != EValidationStates::Not_Applicable && Node.Pin()->GetValidatorState() != EValidationStates::Disabled; })
 					.HAlign(EHorizontalAlignment::HAlign_Center)
 					.VAlign(EVerticalAlignment::VAlign_Center)
 					.OnClicked_Lambda([this]() {ViewLogCallback.ExecuteIfBound(Node.Pin()); return FReply::Handled(); })					
@@ -412,6 +444,22 @@ public:
 						.Text(FText::FromString("View Log"))
 						.Font(FAppStyle::Get().GetFontStyle("SmallFont"))
 					]
+				]
+				+SHorizontalBox::Slot()
+				.Padding(3, 0, 3, 0)
+				[
+					SNew(SButton)
+					.IsEnabled(&FModelInterface::GetInputEnabled)
+					.Visibility_Lambda([this]{ return Node.Pin()->Definition->bCanBeDisabledByUser ? EVisibility::All : EVisibility::Collapsed; })
+					.HAlign(EHorizontalAlignment::HAlign_Center)
+					.VAlign(EVerticalAlignment::VAlign_Center)
+					.OnClicked_Lambda([this]() { ModelInterface->ToggleValidator(Node.Pin()->GetValidatorNameId()); ParentWidget->ReevaluateValidatorSection(); return FReply::Handled(); })
+					[
+						SNew(STextBlock)
+						.Justification(ETextJustify::Center)
+						.Text_Lambda([this]{ return Node.Pin()->Definition->bIsDisabled ? FText::FromString("Enable") : FText::FromString("Disable"); })
+						.Font(FAppStyle::Get().GetFontStyle("SmallFont"))
+					]
 
 				]);	
 		}
@@ -424,13 +472,15 @@ public:
 
 	TWeakPtr<const FValidatorBase> Node;
 	FModelInterface* ModelInterface;
+	SValidatorsWidget* ParentWidget;
 	bool bIsRelevantToCL = false;
 	FOnViewValidatorLog ViewLogCallback;
 };
 
-TSharedRef<ITableRow> SValidatorsWidget::GenerateRow(TWeakPtr<const FValidatorBase> InItem, const TSharedRef<STableViewBase>& OwnerTable) const
+TSharedRef<ITableRow> SValidatorsWidget::GenerateRow(TWeakPtr<const FValidatorBase> InItem, const TSharedRef<STableViewBase>& OwnerTable)
 {
 	return SNew(SValidatorNode, OwnerTable, InItem)
+		.ParentWidget(this)
 		.ModelInterface(ModelInterface)
 		.IsRelevantToCL(InItem.Pin()->IsRelevantToCL())
 		.ViewLogCallback(OnViewLog);
diff --git a/Engine/Source/Programs/SubmitTool/Private/View/Widgets/ValidatorsWidget.h b/Engine/Source/Programs/SubmitTool/Private/View/Widgets/ValidatorsWidget.h
--- a/Engine/Source/Programs/SubmitTool/Private/View/Widgets/ValidatorsWidget.h
+++ b/Engine/Source/Programs/SubmitTool/Private/View/Widgets/ValidatorsWidget.h
@@ -38,6 +38,7 @@ public:
 	void Construct(const FArguments& InArgs);
 	virtual ~SValidatorsWidget();
 
+	void ReevaluateValidatorSection();
 private:
 
 	TSharedRef<SWidget> BuildValidatorsView(bool bListPreSubmitOperations = false);
@@ -52,11 +53,14 @@ private:
 	TArray<FValidatorColumn> Columns;
 
 	TSharedRef<SHeaderRow> ConstructHeadersRow(const FName& GroupName);
-	TSharedRef<ITableRow> GenerateRow(TWeakPtr<const FValidatorBase> InItem, const TSharedRef<STableViewBase>& OwnerTable) const;
+	TSharedRef<ITableRow> GenerateRow(TWeakPtr<const FValidatorBase> InItem, const TSharedRef<STableViewBase>& OwnerTable);
 	
 	void OnColumnSort(const FName& GroupName, EColumnSortPriority::Type InSortPriority, const FName& InColumnId, EColumnSortMode::Type InSortDirection);
 	EColumnSortMode::Type GetSortMode(const FName ColumnId) const;
 
+	FName ActiveSection;
+	FName InactiveSection;
+
 	FName SortByColumn;
     EColumnSortMode::Type SortMode;
 	FOnViewValidatorLog OnViewLog;
