Created a validator that requires the tag 'redirectsaresafe' for changes to PackageRedirectors in ini files.

diff --git a/Engine/Source/Programs/SubmitTool/Private/Logic/Validators/CoreRedirectsValidator.cpp b/Engine/Source/Programs/SubmitTool/Private/Logic/Validators/CoreRedirectsValidator.cpp
new file mode 100644
--- /dev/null
+++ b/Engine/Source/Programs/SubmitTool/Private/Logic/Validators/CoreRedirectsValidator.cpp
@@ -0,0 +1,177 @@
+// Copyright Epic Games, Inc. All Rights Reserved.
+
+#include "CoreRedirectsValidator.h"
+
+#include "Misc/FileHelper.h"
+#include "Misc/Paths.h"
+#include "Configuration/Configuration.h"
+#include "Logging/SubmitToolLog.h"
+
+void FCoreRedirectsValidator::StartAsyncWork(const FString& CLDescription, const TArray<FSourceControlStateRef>& FilteredFilesInCL, const TArray<const FTag*>& Tags) 
+{
+	FilesToValidate.Reset();
+	for (const FSourceControlStateRef& FileInCl : FilteredFilesInCL)
+	{
+		if (FileInCl->IsDeleted())
+		{
+			continue;
+		}
+
+		if (FPaths::GetCleanFilename(FileInCl->GetFilename()).EndsWith(TEXT("Engine.ini"), ESearchCase::IgnoreCase))
+		{
+			FilesToValidate.Add(FileInCl);
+		}
+	}
+
+	TagStatus = Tags;
+
+	TSharedPtr<ISTSourceControlService> SourceControlService = ServiceProvider.Pin()->GetService<ISTSourceControlService>();
+
+	StartAsyncTask([this](const UE::Tasks::FCancellationToken& InCancellationToken)->bool
+		{
+			if (FilesToValidate.IsEmpty())
+			{
+				return true;
+			}
+
+			return DoWork(InCancellationToken);
+		});
+}
+
+bool FCoreRedirectsValidator::DoWork(const UE::Tasks::FCancellationToken& InCancellationToken)
+{
+	// First check the tags to avoid doing unnecessary work if the user has already promised it's safe
+	FString DocumentationURL;
+	for (const FTag* Tag : TagStatus)
+	{
+		if (Tag)
+		{
+			if (Tag->Definition.TagId.Compare(TEXT("#redirectsaresafe"), ESearchCase::IgnoreCase) == 0)
+			{
+				if (Tag->IsEnabled())
+				{
+					return true;
+				}
+				else
+				{
+					DocumentationURL = Tag->Definition.DocumentationUrl;
+				}
+			}
+		}
+	}
+
+	TArray<FSharedBuffer> FileBuffers;
+
+	bool bAnyFilesContainRelevantModifications = false;
+	for (const FSourceControlStateRef& File : FilesToValidate)
+	{
+		if (InCancellationToken.IsCanceled())
+		{
+			return false;
+		}
+
+		if (File->IsAdded())
+		{
+			// This file is new (or potentially moved/renamed). We'll just look to see if it contains redirects.
+			TArray<FString> WorkspaceContents;
+			if (!FFileHelper::LoadFileToStringArray(WorkspaceContents, *File->GetFilename()))
+			{
+				LogFailure(FString::Printf(TEXT("[%s]: Unable to perform validation because we failed to load '%s'."), *GetValidatorName(), *File->GetFilename()));
+				return false;
+			}
+			else
+			{
+				if (ContainsAnyRedirects(WorkspaceContents))
+				{
+					LogFailure(FString::Printf(TEXT("[%s]: Your changelist includes changes to PackageRedirectors via ini file. Please read the linked documentation"), *GetValidatorName()));
+					LogFailure(FString::Printf(TEXT("[%s]: %s"), *GetValidatorName(), *DocumentationURL));
+					LogFailure(FString::Printf(TEXT("[%s]: If your changes are safe, add #redirectsaresafe"), *GetValidatorName()));
+					return false;
+				}
+			}
+		}
+		else if (ServiceProvider.Pin()->GetService<ISTSourceControlService>()->DownloadFiles(File->GetFilename(), FileBuffers).GetResult().bRequestSucceed)
+		{
+			if (InCancellationToken.IsCanceled())
+			{
+				return false;
+			}
+
+			if (FileBuffers.Num() != 1 || FileBuffers[0].GetSize() == 0)
+			{
+				LogFailure(FString::Printf(TEXT("[%s]: Unable to perform validation because we received bad data from version control."), *GetValidatorName()));
+				return false;
+			}
+
+			const FAnsiStringView DataView(reinterpret_cast<const char*>(FileBuffers[0].GetData()));
+			TSet<FAnsiStringView> DepotContents;
+			PopulateSetFromStringViewOfFile(DataView, DepotContents);
+
+			TArray<FString> WorkspaceContents;
+			if (!FFileHelper::LoadFileToStringArray(WorkspaceContents, *File->GetFilename()))
+			{
+				LogFailure(FString::Printf(TEXT("[%s]: Unable to perform validation because we failed to load '%s'."), *GetValidatorName(), *File->GetFilename()));
+				return false;
+			}
+
+			if (ContainsModifiedRedirects(DepotContents, WorkspaceContents))
+			{
+				LogFailure(FString::Printf(TEXT("[%s]: Your changelist includes changes to PackageRedirectors via ini file. Please read the linked documentation"), *GetValidatorName()));
+				LogFailure(FString::Printf(TEXT("[%s]: %s"), *GetValidatorName(), *DocumentationURL));
+				LogFailure(FString::Printf(TEXT("[%s]: If your changes are safe, add #redirectsaresafe"), *GetValidatorName()));
+				return false;
+			}
+		}
+		else
+		{
+			LogFailure(FString::Printf(TEXT("[%s]: Unable to perform validation because we could not download the unedited file ('%s') from version control."), *GetValidatorName(), *File->GetFilename()));
+		}
+
+		FileBuffers.Reset();
+	}
+
+	return true;
+
+}
+
+void FCoreRedirectsValidator::PopulateSetFromStringViewOfFile(FAnsiStringView InFile, TSet<FAnsiStringView>& OutSet)
+{
+	int32 LineTerminator = 0;
+	// Iterate over the total string and repeatedly split it on line endings into a new line for OutSet and a remainder
+	while (InFile.FindChar('\n', LineTerminator))
+	{
+		OutSet.Add(InFile.Left(LineTerminator));
+		InFile = InFile.RightChop(LineTerminator + 1);
+	}
+}
+
+bool FCoreRedirectsValidator::ContainsModifiedRedirects(const TSet<FAnsiStringView>& DepotContents, const TArray<FString>& WorkspaceContents)
+{
+	for (const FString& Line : WorkspaceContents)
+	{
+		if (!DepotContents.Contains(FAnsiStringView(TCHAR_TO_ANSI(*Line), Line.Len())))
+		{
+			FStringView LineView(Line);
+			LineView = LineView.TrimStart();
+			if (LineView.StartsWith(TEXT("+PackageRedirects")))
+			{
+				return true;
+			}
+		}
+	}
+	return false;
+}
+
+bool FCoreRedirectsValidator::ContainsAnyRedirects(const TArray<FString>& FileContents)
+{
+	for (const FString& Line : FileContents)
+	{
+		FStringView LineView(Line);
+		LineView = LineView.TrimStart();
+		if (LineView.StartsWith(TEXT("+PackageRedirects")))
+		{
+			return true;
+		}
+	}
+	return false;
+}
diff --git a/Engine/Source/Programs/SubmitTool/Private/Logic/Validators/CoreRedirectsValidator.h b/Engine/Source/Programs/SubmitTool/Private/Logic/Validators/CoreRedirectsValidator.h
new file mode 100644
--- /dev/null
+++ b/Engine/Source/Programs/SubmitTool/Private/Logic/Validators/CoreRedirectsValidator.h
@@ -0,0 +1,33 @@
+// Copyright Epic Games, Inc. All Rights Reserved.
+
+#pragma once
+
+#include "ValidatorBaseAsync.h"
+#include "Tasks/Task.h"
+#include "ISourceControlChangelistState.h"
+#include "ISourceControlProvider.h"
+#include "Logic/Services/Interfaces/ISTSourceControlService.h"
+
+class FCoreRedirectsValidator : public FValidatorBaseAsync
+{
+public:
+	using FValidatorBaseAsync::FValidatorBaseAsync;
+	virtual void StartAsyncWork(const FString& CLDescription, const TArray<FSourceControlStateRef>& FilteredFilesInCL, const TArray<const FTag*>& Tags) override;
+
+public:
+
+	virtual const FString& GetValidatorTypeName() const override { return SubmitToolParseConstants::CoreRedirectsValidator; }
+
+private:
+
+	TArray<const FTag*> TagStatus;
+	TArray<FSourceControlStateRef> FilesToValidate;
+
+	bool DoWork(const UE::Tasks::FCancellationToken& InCancellationToken);
+
+	void PopulateSetFromStringViewOfFile(FAnsiStringView InFile, TSet<FAnsiStringView>& OutSet);
+
+	virtual bool ContainsModifiedRedirects(const TSet<FAnsiStringView>& DepotContents, const TArray<FString>& WorkspaceContents);
+
+	bool ContainsAnyRedirects(const TArray<FString>& FileContents);
+};
\ No newline at end of file
diff --git a/Engine/Source/Programs/SubmitTool/Private/Logic/Validators/ValidatorBase.h b/Engine/Source/Programs/SubmitTool/Private/Logic/Validators/ValidatorBase.h
--- a/Engine/Source/Programs/SubmitTool/Private/Logic/Validators/ValidatorBase.h
+++ b/Engine/Source/Programs/SubmitTool/Private/Logic/Validators/ValidatorBase.h
@@ -33,6 +33,7 @@ namespace SubmitToolParseConstants
 	const FString JsonValidator = TEXT("JsonValidator");
 	const FString InvalidateNode = TEXT("Invalidated");
 	const FString VirtualizationToolOp = TEXT("VirtualizationTool");
+	const FString CoreRedirectsValidator = TEXT("CoreRedirectsValidator");
 };
 
 UENUM()
diff --git a/Engine/Source/Programs/SubmitTool/Private/Models/ModelInterface.cpp b/Engine/Source/Programs/SubmitTool/Private/Models/ModelInterface.cpp
--- a/Engine/Source/Programs/SubmitTool/Private/Models/ModelInterface.cpp
+++ b/Engine/Source/Programs/SubmitTool/Private/Models/ModelInterface.cpp
@@ -14,6 +14,7 @@
 #include "Async/Async.h"
 #include "Misc/StringOutputDevice.h"
 
+#include "Logic/Validators/CoreRedirectsValidator.h"
 #include "Logic/Validators/PreflightValidator.h"
 #include "Logic/Validators/TagValidator.h"
 #include "Logic/Validators/UBTValidator.h"
@@ -252,6 +253,10 @@ void FModelInterface::ParseValidators() const
 		{
 			Tasks.Add(MakeShared<FJsonValidator>(DefinitionPair.Key, Parameters, ServiceProvider.ToSharedRef(), DefinitionPair.Value));
 		}
+		else if (TaskDefinition.Type.Equals(SubmitToolParseConstants::CoreRedirectsValidator, ESearchCase::IgnoreCase))
+		{
+			Tasks.Add(MakeShared<FCoreRedirectsValidator>(DefinitionPair.Key, Parameters, ServiceProvider.ToSharedRef(), DefinitionPair.Value));
+		}
 		else
 		{
 			UE_LOG(LogSubmitTool, Error, TEXT("[%s] is not a recognized validator type and has not been activated."), *DefinitionPair.Key.ToString());
