Added support for StudioTelemetry in submit tool

diff --git a/Engine/Source/Programs/SubmitTool/Private/Platforms/Linux/SubmitToolMainLinux.cpp b/Engine/Source/Programs/SubmitTool/Private/Platforms/Linux/SubmitToolMainLinux.cpp
--- a/Engine/Source/Programs/SubmitTool/Private/Platforms/Linux/SubmitToolMainLinux.cpp
+++ b/Engine/Source/Programs/SubmitTool/Private/Platforms/Linux/SubmitToolMainLinux.cpp
@@ -5,7 +5,7 @@
 
 int RunSubmitToolWrapper(const TCHAR* Commandline)
 {
-    return RunSubmitTool(Commandline, FGuid());
+    return RunSubmitTool(Commandline);
 }
 
 int main(int argc, char *argv[])
diff --git a/Engine/Source/Programs/SubmitTool/Private/Platforms/Mac/SubmitToolMainMac.cpp b/Engine/Source/Programs/SubmitTool/Private/Platforms/Mac/SubmitToolMainMac.cpp
--- a/Engine/Source/Programs/SubmitTool/Private/Platforms/Mac/SubmitToolMainMac.cpp
+++ b/Engine/Source/Programs/SubmitTool/Private/Platforms/Mac/SubmitToolMainMac.cpp
@@ -6,7 +6,7 @@
 
 int RunSubmitToolWrapper(const TCHAR* Commandline)
 {
-	return RunSubmitTool(Commandline, FGuid());
+	return RunSubmitTool(Commandline);
 }
 
 int main(int argc, char *argv[])
diff --git a/Engine/Source/Programs/SubmitTool/Private/Platforms/Windows/SubmitToolMainWindows.cpp b/Engine/Source/Programs/SubmitTool/Private/Platforms/Windows/SubmitToolMainWindows.cpp
--- a/Engine/Source/Programs/SubmitTool/Private/Platforms/Windows/SubmitToolMainWindows.cpp
+++ b/Engine/Source/Programs/SubmitTool/Private/Platforms/Windows/SubmitToolMainWindows.cpp
@@ -23,9 +23,7 @@ int WINAPI WinMain( _In_ HINSTANCE hInInstance, _In_opt_ HINSTANCE hPrevInstance
 		FString CmdLine = FCommandLine::BuildFromArgV(nullptr, __argc, __argv, "");
 		FCommandLine::Set(*CmdLine);
 
-		FGuid guid = FGuid::NewGuid();
-
-		RunSubmitTool(*CmdLine, guid);
+		RunSubmitTool(*CmdLine);
 		GIsGuarded = 0;
 	}
 #if !PLATFORM_SEH_EXCEPTIONS_DISABLED
diff --git a/Engine/Source/Programs/SubmitTool/Private/SubmitToolApp.cpp b/Engine/Source/Programs/SubmitTool/Private/SubmitToolApp.cpp
--- a/Engine/Source/Programs/SubmitTool/Private/SubmitToolApp.cpp
+++ b/Engine/Source/Programs/SubmitTool/Private/SubmitToolApp.cpp
@@ -33,7 +33,7 @@ IMPLEMENT_APPLICATION(SubmitTool, "SubmitTool");
 #define LOCTEXT_NAMESPACE "SubmitTool"
 
 
-int RunSubmitTool(const TCHAR* CommandLine, const FGuid& InSessionID)
+int RunSubmitTool(const TCHAR* CommandLine)
 {
 	FTaskTagScope TaskTagScope(ETaskTag::EGameThread);
 
@@ -82,7 +82,7 @@ int RunSubmitTool(const TCHAR* CommandLine, const FGuid& InSessionID)
 		FSubmitToolParametersBuilder ParametersBuilder(ParameterFile);
 		FSubmitToolParameters Parameters = ParametersBuilder.Build();
 
-		FTelemetryService::Init(Parameters.Telemetry.Url, InSessionID);
+		FTelemetryService::Init(Parameters.Telemetry);
 
 		// Create a new instance of model Interface so that UI can communicate
 		const TUniquePtr<FModelInterface> ModelInterface = MakeUnique<FModelInterface>(Parameters);
diff --git a/Engine/Source/Programs/SubmitTool/Private/Telemetry/SubmitToolStudioTelemetry.cpp b/Engine/Source/Programs/SubmitTool/Private/Telemetry/SubmitToolStudioTelemetry.cpp
new file mode 100644
--- /dev/null
+++ b/Engine/Source/Programs/SubmitTool/Private/Telemetry/SubmitToolStudioTelemetry.cpp
@@ -0,0 +1,51 @@
+// Copyright Epic Games, Inc. All Rights Reserved.
+
+#include "SubmitToolStudioTelemetry.h"
+#include "Logging/SubmitToolLog.h"
+#include "Version/AppVersion.h"
+#include "AnalyticsEventAttribute.h"
+#include "StudioTelemetry.h"
+
+const FName SpanId = TEXT("SubmitTool.Session");
+
+FSubmitToolStudioTelemetry::FSubmitToolStudioTelemetry()
+{
+	FStudioTelemetry::Get().StartSession();	
+	FStudioTelemetry::Get().StartSpan(SpanId);
+}
+
+FSubmitToolStudioTelemetry::~FSubmitToolStudioTelemetry()
+{
+	FStudioTelemetry::Get().EndSpan(SpanId);
+	FStudioTelemetry::Get().EndSession();
+}
+
+void FSubmitToolStudioTelemetry::Start(const FString& InCurrentStream) const
+{
+	FStudioTelemetry::Get().RecordEvent(
+		TEXT("SubmitTool.StandAlone.Start"),
+		MakeAnalyticsEventAttributeArray(
+			TEXT("Version"), FAppVersion::GetVersion(),
+			TEXT("Stream"), InCurrentStream
+		)
+	);
+}
+
+void FSubmitToolStudioTelemetry::BlockFlush(float InTimeout) const
+{
+	FStudioTelemetry::Get().FlushEvents();
+}
+void FSubmitToolStudioTelemetry::CustomEvent(const FString& InEventId, const TArray<FAnalyticsEventAttribute>& InAttribs) const
+{
+	FStudioTelemetry::Get().RecordEvent(InEventId, InAttribs);
+}
+
+void FSubmitToolStudioTelemetry::SubmitSucceeded() const
+{
+	FStudioTelemetry::Get().RecordEvent(
+		TEXT("SubmitTool.StandAlone.Submit.Succeeded"),
+		MakeAnalyticsEventAttributeArray(
+			TEXT("Version"), FAppVersion::GetVersion()
+		)
+	);
+}
\ No newline at end of file
diff --git a/Engine/Source/Programs/SubmitTool/Private/Telemetry/SubmitToolStudioTelemetry.h b/Engine/Source/Programs/SubmitTool/Private/Telemetry/SubmitToolStudioTelemetry.h
new file mode 100644
--- /dev/null
+++ b/Engine/Source/Programs/SubmitTool/Private/Telemetry/SubmitToolStudioTelemetry.h
@@ -0,0 +1,16 @@
+// Copyright Epic Games, Inc. All Rights Reserved.
+
+#pragma once
+
+#include "ITelemetry.h"
+
+class FSubmitToolStudioTelemetry : public ITelemetry
+{
+public:
+	FSubmitToolStudioTelemetry();
+	virtual ~FSubmitToolStudioTelemetry();
+	virtual void Start(const FString& InCurrentStream) const override;
+	virtual void BlockFlush(float InTimeout) const override;
+	virtual void SubmitSucceeded() const override;
+	virtual void CustomEvent(const FString& InEventId, const TArray<FAnalyticsEventAttribute>& InAttribs) const override;
+};
\ No newline at end of file
diff --git a/Engine/Source/Programs/SubmitTool/Private/Telemetry/TelemetryService.cpp b/Engine/Source/Programs/SubmitTool/Private/Telemetry/TelemetryService.cpp
--- a/Engine/Source/Programs/SubmitTool/Private/Telemetry/TelemetryService.cpp
+++ b/Engine/Source/Programs/SubmitTool/Private/Telemetry/TelemetryService.cpp
@@ -3,8 +3,10 @@
 #include "TelemetryService.h"
 
 #include "Logging/SubmitToolLog.h"
+#include "Parameters/SubmitToolParameters.h"
 #include "NullTelemetry.h"
 #include "StandAloneTelemetry.h"
+#include "SubmitToolStudioTelemetry.h"
 
 FCriticalSection FTelemetryService::InstanceCriticalSection;
 TSharedPtr<ITelemetry> FTelemetryService::TelemetryInstance = nullptr;
@@ -14,14 +16,27 @@ const TSharedPtr<ITelemetry>& FTelemetryService::Get()
 	return TelemetryInstance;
 }
 
-void FTelemetryService::Init(const FString& InUrl, const FGuid& InSessionID)
+void FTelemetryService::Init(const FTelemetryParameters& InTelemetryParameters)
 {
-	if (InUrl.IsEmpty())
+	if (InTelemetryParameters.Instance.Equals(TEXT("StandAlone"), ESearchCase::IgnoreCase))
+	{
+		if (InTelemetryParameters.Url.IsEmpty())
+		{
+			Set(MakeShared<FNullTelemetry>());
+		}
+		else
+		{
+			Set(MakeShared<FStandAloneTelemetry>(InTelemetryParameters.Url, FGuid()));
+		}
+	}
+	else if (InTelemetryParameters.Instance.Equals(TEXT("StudioTelemetry"), ESearchCase::IgnoreCase))
+	{
+		Set(MakeShared<FSubmitToolStudioTelemetry>());
+	}
+	else
 	{
 		Set(MakeShared<FNullTelemetry>());
 	}
-
-	Set(MakeShared<FStandAloneTelemetry>(InUrl, InSessionID));
 }
 
 void FTelemetryService::Set(TSharedPtr<ITelemetry> InInstance)
diff --git a/Engine/Source/Programs/SubmitTool/Private/Telemetry/TelemetryService.h b/Engine/Source/Programs/SubmitTool/Private/Telemetry/TelemetryService.h
--- a/Engine/Source/Programs/SubmitTool/Private/Telemetry/TelemetryService.h
+++ b/Engine/Source/Programs/SubmitTool/Private/Telemetry/TelemetryService.h
@@ -7,11 +7,13 @@
 #include "StandAloneTelemetry.h"
 #include "ITelemetry.h"
 
+struct FTelemetryParameters;
+
 class FTelemetryService
 {
 public:
 	static const TSharedPtr<ITelemetry> & Get();
-	static void Init(const FString& InUrl, const FGuid& InSessionID);
+	static void Init(const FTelemetryParameters& InTelemetryParameters);
 	static void Shutdown();
 	static void BlockFlush(float InTimeout);
 private:
diff --git a/Engine/Source/Programs/SubmitTool/Private/Version/AppVersionDefines.h b/Engine/Source/Programs/SubmitTool/Private/Version/AppVersionDefines.h
--- a/Engine/Source/Programs/SubmitTool/Private/Version/AppVersionDefines.h
+++ b/Engine/Source/Programs/SubmitTool/Private/Version/AppVersionDefines.h
@@ -20,8 +20,8 @@
 #endif
 
 #define SUBMIT_TOOL_VERSION_MAJOR 6
-#define SUBMIT_TOOL_VERSION_MINOR 1
-#define SUBMIT_TOOL_VERSION_PATCH 2
+#define SUBMIT_TOOL_VERSION_MINOR 2
+#define SUBMIT_TOOL_VERSION_PATCH 0
 
 #ifndef SUBMIT_TOOL_CHANGELIST
 #define SUBMIT_TOOL_CHANGELIST_STRING TEXT("NoCL")
diff --git a/Engine/Source/Programs/SubmitTool/Public/SubmitToolApp.h b/Engine/Source/Programs/SubmitTool/Public/SubmitToolApp.h
--- a/Engine/Source/Programs/SubmitTool/Public/SubmitToolApp.h
+++ b/Engine/Source/Programs/SubmitTool/Public/SubmitToolApp.h
@@ -7,5 +7,5 @@
 /**
  * Run the SubmitTool .
  */
-int RunSubmitTool(const TCHAR* Commandline, const FGuid& SessionID);
+int RunSubmitTool(const TCHAR* Commandline);
 FString GetUserPrefsPath();
diff --git a/Engine/Source/Programs/SubmitTool/SubmitTool.Build.cs b/Engine/Source/Programs/SubmitTool/SubmitTool.Build.cs
--- a/Engine/Source/Programs/SubmitTool/SubmitTool.Build.cs
+++ b/Engine/Source/Programs/SubmitTool/SubmitTool.Build.cs
@@ -26,6 +26,7 @@ public class SubmitTool : ModuleRules
 				"Slate",
 				"SlateCore",
 				"StandaloneRenderer",
+				"StudioTelemetry",
 				"ToolWidgets",
 				"SourceControl",
 				"DesktopPlatform", // Open file dialog
