Adding button for disabling validators as a means to bypass acceptable validator errors
Added popup warning for disabling validators
Decluttering the UI:
Merged Force run, Run and Stop buttons into the same button (Use left Alt for force)
Removed View Log button

diff --git a/Engine/Source/Programs/SubmitTool/Private/Logic/DialogFactory.cpp b/Engine/Source/Programs/SubmitTool/Private/Logic/DialogFactory.cpp
--- a/Engine/Source/Programs/SubmitTool/Private/Logic/DialogFactory.cpp
+++ b/Engine/Source/Programs/SubmitTool/Private/Logic/DialogFactory.cpp
@@ -6,7 +6,7 @@
 #include "Framework/Application/SlateApplication.h"
 
 EDialogFactoryResult FDialogFactory::ShowDialog(const FText& Title, const FText& Description,
-	const TArray<FString>& AvailableButtons, const TSharedPtr<SWidget> AdditionalContent)
+	const TArray<FString>& AvailableButtons, const TSharedPtr<SWidget> AdditionalContent, TFunction<bool(size_t)> InIsBtnEnabled)
 {	
 	TSharedPtr<SWindow> Window = SNew(SWindow)
 	.Title(Title)
@@ -27,7 +27,8 @@ EDialogFactoryResult FDialogFactory::ShowDialog(const FText& Title, const FText&
 		.DescriptionText(Description)
 		.Buttons(AvailableButtons)
 		.ResultCallback(ResultsCallback)
-		.AdditionalContent(AdditionalContent);
+		.AdditionalContent(AdditionalContent)
+		.IsBtnEnabled(InIsBtnEnabled);
 
 	Window->SetContent(ConfirmWidget.ToSharedRef());
 
diff --git a/Engine/Source/Programs/SubmitTool/Private/Logic/DialogFactory.h b/Engine/Source/Programs/SubmitTool/Private/Logic/DialogFactory.h
--- a/Engine/Source/Programs/SubmitTool/Private/Logic/DialogFactory.h
+++ b/Engine/Source/Programs/SubmitTool/Private/Logic/DialogFactory.h
@@ -20,7 +20,7 @@ enum class EDialogFactoryResult : size_t
 class FDialogFactory
 {
 public:
-	static EDialogFactoryResult ShowDialog(const FText& Title, const FText& Description, const TArray<FString>& AvailableButtons, const TSharedPtr<SWidget> AdditionalContent = nullptr);
+	static EDialogFactoryResult ShowDialog(const FText& Title, const FText& Description, const TArray<FString>& AvailableButtons, const TSharedPtr<SWidget> AdditionalContent = nullptr, TFunction<bool(size_t)> InIsBtnEnabled = nullptr);
 
 	/**
 	 * @brief Shows a dialog screen with the Cancel and Confirm buttons
diff --git a/Engine/Source/Programs/SubmitTool/Private/Logic/Validators/ValidatorRunExecutable.cpp b/Engine/Source/Programs/SubmitTool/Private/Logic/Validators/ValidatorRunExecutable.cpp
--- a/Engine/Source/Programs/SubmitTool/Private/Logic/Validators/ValidatorRunExecutable.cpp
+++ b/Engine/Source/Programs/SubmitTool/Private/Logic/Validators/ValidatorRunExecutable.cpp
@@ -51,6 +51,7 @@ bool FValidatorRunExecutable::Validate(const FString& InCLDescription, const TAr
 	const FValidatorRunExecutableDefinition* TypedDefinition = GetTypedDefinition<FValidatorRunExecutableDefinition>();
 
 	FString FinalArgs{ FConfiguration::Substitute(TypedDefinition->ExecutableArguments) };
+	FinalArgs.ReplaceInline(TEXT("$(TaskName)"), *GetValidatorName());
 
 	TArray<FString> Files;
 
diff --git a/Engine/Source/Programs/SubmitTool/Private/View/SubmitToolStyle.cpp b/Engine/Source/Programs/SubmitTool/Private/View/SubmitToolStyle.cpp
--- a/Engine/Source/Programs/SubmitTool/Private/View/SubmitToolStyle.cpp
+++ b/Engine/Source/Programs/SubmitTool/Private/View/SubmitToolStyle.cpp
@@ -107,6 +107,8 @@ FSubmitToolStyle::FSubmitToolStyle() : FSlateStyleSet(StyleName)
 	Set("RichTextBlock.ItalicHighlight", FTextBlockStyle(NormalText)
 		.SetFont(DEFAULT_FONT("Italic", FCoreStyle::RegularTextSize))
 		.SetColorAndOpacity(FLinearColor(1.0f, 1.0f, 1.0f)));
+	Set("RichTextBlock.Red", FTextBlockStyle(NormalText)
+		.SetColorAndOpacity(FLinearColor(1.0f, 0.1f, 0.1f)));
 
 	Set("ValidatorStateWarning", FStyleColors::Warning.GetSpecifiedColor());
 	Set("ValidatorStateFail", FStyleColors::Error.GetSpecifiedColor());
diff --git a/Engine/Source/Programs/SubmitTool/Private/View/Widgets/ConfirmDialogWidget.cpp b/Engine/Source/Programs/SubmitTool/Private/View/Widgets/ConfirmDialogWidget.cpp
--- a/Engine/Source/Programs/SubmitTool/Private/View/Widgets/ConfirmDialogWidget.cpp
+++ b/Engine/Source/Programs/SubmitTool/Private/View/Widgets/ConfirmDialogWidget.cpp
@@ -11,7 +11,7 @@
 void SConfirmDialogWidget::Construct(const FArguments& InArgs)
 {
 	ResultsCallback = InArgs._ResultCallback;
-	
+	IsBtnEnabled = InArgs._IsBtnEnabled;
 	TSharedPtr<SHorizontalBox> ButtonBox;
 	
 	ChildSlot
@@ -55,7 +55,6 @@ void SConfirmDialogWidget::Construct(const FArguments& InArgs)
 		]
 	];
 
-	const FOnResult& ResultCallback = InArgs._ResultCallback;
 	
 	for(size_t i=0;i<InArgs._Buttons.Num();++i)
 	{
@@ -72,6 +71,7 @@ void SConfirmDialogWidget::ConstructButton(const TSharedRef<SHorizontalBox> Cont
 		[
 			SNew(SButton)
 			.Text(FText::FromString(ButtonText))
+			.IsEnabled_Lambda([Idx, this]{ return IsBtnEnabled == nullptr || IsBtnEnabled(Idx);})
 			.OnClicked_Lambda([this, Idx]()
 			{
 				ResultsCallback.ExecuteIfBound(Idx);
diff --git a/Engine/Source/Programs/SubmitTool/Private/View/Widgets/ConfirmDialogWidget.h b/Engine/Source/Programs/SubmitTool/Private/View/Widgets/ConfirmDialogWidget.h
--- a/Engine/Source/Programs/SubmitTool/Private/View/Widgets/ConfirmDialogWidget.h
+++ b/Engine/Source/Programs/SubmitTool/Private/View/Widgets/ConfirmDialogWidget.h
@@ -16,6 +16,7 @@ public:
 		SLATE_ARGUMENT(FText, DescriptionText)
 		SLATE_ARGUMENT(TSharedPtr<SWidget>, AdditionalContent)
 		SLATE_ARGUMENT(TArray<FString>, Buttons)
+		SLATE_ARGUMENT(TFunction<bool(size_t)>, IsBtnEnabled)
 		SLATE_EVENT(FOnResult, ResultCallback)
 	SLATE_END_ARGS()
 
@@ -23,6 +24,7 @@ public:
 
 private:
 	FOnResult ResultsCallback;
+	TFunction<bool(size_t)> IsBtnEnabled = nullptr;
 	
 	void ConstructButton(
 		const TSharedRef<SHorizontalBox> Container,
diff --git a/Engine/Source/Programs/SubmitTool/Private/View/Widgets/ValidatorsWidget.cpp b/Engine/Source/Programs/SubmitTool/Private/View/Widgets/ValidatorsWidget.cpp
--- a/Engine/Source/Programs/SubmitTool/Private/View/Widgets/ValidatorsWidget.cpp
+++ b/Engine/Source/Programs/SubmitTool/Private/View/Widgets/ValidatorsWidget.cpp
@@ -9,6 +9,7 @@
 #include "Widgets/Layout/SExpandableArea.h"
 #include "Widgets/Input/SButton.h"
 #include "Widgets/Input/SComboBox.h"
+#include "Widgets/Input/SCheckBox.h"
 #include "Widgets/Text/STextBlock.h"
 
 #include "Models/SubmitToolUserPrefs.h"
@@ -30,9 +31,9 @@ void SValidatorsWidget::Construct(const FArguments& InArgs)
 {
 	Columns = {
 		FValidatorColumn(ValidatorsColumns::Name, 1, true, [](TWeakPtr<const FValidatorBase> In) {return In.Pin()->GetValidatorName(); }),
-		FValidatorColumn(ValidatorsColumns::Required, 80, false, [](TWeakPtr<const FValidatorBase> In) {return In.Pin()->Definition->IsRequired ? TEXT("Yes") : TEXT("No"); }),
-		FValidatorColumn(ValidatorsColumns::Status, 100, false, [](TWeakPtr<const FValidatorBase> In) {return In.Pin()->GetStatusText(); }),
-		FValidatorColumn(ValidatorsColumns::Actions, 350, false, nullptr),
+		FValidatorColumn(ValidatorsColumns::Required, 130, false, [](TWeakPtr<const FValidatorBase> In) {return In.Pin()->Definition->IsRequired ? TEXT("Yes") : TEXT("No"); }),
+		FValidatorColumn(ValidatorsColumns::Status, 130, false, [](TWeakPtr<const FValidatorBase> In) {return In.Pin()->GetStatusText(); }),
+		FValidatorColumn(ValidatorsColumns::Actions, 200, false, nullptr),
 	};
 
 	ModelInterface = InArgs._ModelInterface;
@@ -357,31 +358,39 @@ public:
 		else if(InColumnName == ValidatorsColumns::Status)
 		{
 			box->SetContent(
-				SNew(STextBlock)
-				.Text_Lambda([this]() { return FText::FromString(Node.Pin()->GetStatusText()); })
-				.ColorAndOpacity_Lambda([this]{
-					switch(Node.Pin()->GetValidatorState())
-					{
-					case EValidationStates::Failed:
-					case EValidationStates::Timeout:
-						return Node.Pin()->Definition->IsRequired ? FAppStyle::GetColor("ValidatorStateFail") : FAppStyle::GetColor("ValidatorStateWarning");
-						break;
-					case EValidationStates::Valid:
-					case EValidationStates::Skipped:
-					case EValidationStates::Not_Applicable:
-						return FAppStyle::GetColor("ValidatorStateSuccess");
-						break;
-					case EValidationStates::Running:
-						return FAppStyle::GetColor("ValidatorStateWarning");
-						break;
-					case EValidationStates::Disabled:
-						return FAppStyle::GetColor("ValidatorStateDisabled");
-						break;
-					case EValidationStates::Not_Run:
-					default:
-						return FAppStyle::GetColor("ValidatorStateNormal");
-						break;
-					}})
+				SNew(SButton)
+				.ButtonStyle(FAppStyle::Get(), "InvisibleButton")
+				.IsFocusable(false)
+				.ToolTipText(FText::FromString(TEXT("View Log")))
+				.OnClicked_Lambda([this] {	ViewLogCallback.ExecuteIfBound(Node.Pin()); return FReply::Handled();	})
+				[
+					SNew(STextBlock)
+					.Cursor(EMouseCursor::Hand)
+					.Text_Lambda([this]() { return FText::FromString(Node.Pin()->GetStatusText()); })
+					.ColorAndOpacity_Lambda([this]{
+						switch(Node.Pin()->GetValidatorState())
+						{
+						case EValidationStates::Failed:
+						case EValidationStates::Timeout:
+							return Node.Pin()->Definition->IsRequired ? FAppStyle::GetColor("ValidatorStateFail") : FAppStyle::GetColor("ValidatorStateWarning");
+							break;
+						case EValidationStates::Valid:
+						case EValidationStates::Skipped:
+						case EValidationStates::Not_Applicable:
+							return FAppStyle::GetColor("ValidatorStateSuccess");
+							break;
+						case EValidationStates::Running:
+							return FAppStyle::GetColor("ValidatorStateWarning");
+							break;
+						case EValidationStates::Disabled:
+							return FAppStyle::GetColor("ValidatorStateDisabled");
+							break;
+						case EValidationStates::Not_Run:
+						default:
+							return FAppStyle::GetColor("ValidatorStateNormal");
+							break;
+						}})
+				]
 			);
 		}
 		else if(InColumnName == ValidatorsColumns::Actions)
@@ -396,52 +405,37 @@ public:
 					.VAlign(EVerticalAlignment::VAlign_Center)
 					.OnClicked_Lambda([this]()
 					{
-						ModelInterface->ValidateSingle(Node.Pin()->GetValidatorNameId());
-						return FReply::Handled();
-					})
-					[
-						SNew(STextBlock)
-						.Justification(ETextJustify::Center)
-						.IsEnabled_Lambda([this](){	return !Node.Pin()->GetIsRunningOrQueued();})
-						.Text_Lambda([this](){ return  FText::FromString("Force Run"); })
-						.Font(FAppStyle::Get().GetFontStyle("SmallFont"))
-					]
-				]
-				+SHorizontalBox::Slot()
-				.Padding(3,0,0,0)
-				[
-					SNew(SButton)
-					.IsEnabled_Lambda([this]() {return ModelInterface->GetInputEnabled() && Node.Pin()->GetValidatorState() != EValidationStates::Disabled && bIsRelevantToCL; })
-					.HAlign(EHorizontalAlignment::HAlign_Center)
-					.VAlign(EVerticalAlignment::VAlign_Center)
-					.OnClicked_Lambda([this]()
-					{
-						if(Node.Pin()->GetIsRunningOrQueued()){
+						if (Node.Pin()->GetIsRunningOrQueued())
+						{
 							ModelInterface->CancelValidations(Node.Pin()->GetValidatorNameId(), true);
-						}else{
-							ModelInterface->ValidateSingle(Node.Pin()->GetValidatorNameId(), false);
+						}
+						else
+						{
+							ModelInterface->ValidateSingle(Node.Pin()->GetValidatorNameId(), FSlateApplication::Get().GetModifierKeys().IsLeftAltDown());
 						}
 						return FReply::Handled();
 					})
 					[
 						SNew(STextBlock)
 						.Justification(ETextJustify::Center)
-						.Text_Lambda([this](){ return Node.Pin()->GetIsRunningOrQueued() ? FText::FromString("Stop") : FText::FromString("Run"); })
-						.Font(FAppStyle::Get().GetFontStyle("SmallFont"))
-					]
-				]
-				+SHorizontalBox::Slot()
-				.Padding(3, 0, 3, 0)
-				[
-					SNew(SButton)
-					.IsEnabled_Lambda([this]() { return Node.Pin()->GetValidatorState() != EValidationStates::Not_Run && Node.Pin()->GetValidatorState() != EValidationStates::Skipped && Node.Pin()->GetValidatorState() != EValidationStates::Not_Applicable && Node.Pin()->GetValidatorState() != EValidationStates::Disabled; })
-					.HAlign(EHorizontalAlignment::HAlign_Center)
-					.VAlign(EVerticalAlignment::VAlign_Center)
-					.OnClicked_Lambda([this]() {ViewLogCallback.ExecuteIfBound(Node.Pin()); return FReply::Handled(); })					
-					[
-						SNew(STextBlock)
-						.Justification(ETextJustify::Center)
-						.Text(FText::FromString("View Log"))
+						.Text_Lambda([this]() 
+						{ 
+							if (Node.Pin()->GetIsRunningOrQueued())
+							{
+								return FText::FromString("Stop");
+							}
+							else
+							{
+								if(FSlateApplication::Get().GetModifierKeys().IsLeftAltDown())
+								{
+									return FText::FromString("Force Run");
+								}
+								else
+								{
+									return FText::FromString("Run");
+								}
+							}
+						})
 						.Font(FAppStyle::Get().GetFontStyle("SmallFont"))
 					]
 				]
@@ -453,7 +447,34 @@ public:
 					.Visibility_Lambda([this]{ return Node.Pin()->Definition->bCanBeDisabledByUser ? EVisibility::All : EVisibility::Collapsed; })
 					.HAlign(EHorizontalAlignment::HAlign_Center)
 					.VAlign(EVerticalAlignment::VAlign_Center)
-					.OnClicked_Lambda([this]() { ModelInterface->ToggleValidator(Node.Pin()->GetValidatorNameId()); ParentWidget->ReevaluateValidatorSection(); return FReply::Handled(); })
+					.OnClicked_Lambda([this]() 
+					{
+						bool bConfirmed = ParentWidget->bDontShowDisablePopup;
+
+						if (!bConfirmed)
+						{
+							FText Title = FText::FromString(FString::Printf(TEXT("Disable %s"), *Node.Pin()->GetValidatorName()));
+							FText MessageBody = FText::FromString(FString::Printf(TEXT("You are <RichTextBlock.TextHighlight>disabling %s.</>\n\nThis means some potentially necessary checks may not be performed, and %s will be recorded in this CL as <RichTextBlock.Red>DISABLED</>.\n\nIf this CL causes breakages you will be asked to provide justifications for having disabled validations.\n"), *Node.Pin()->GetValidatorName(), *Node.Pin()->GetValidatorName()));
+
+							TSharedRef<SWidget> AdditionalContent = 
+								SNew(SCheckBox)
+								.OnCheckStateChanged_Lambda([this](ECheckBoxState InChkState) { ParentWidget->bDontShowDisablePopup = InChkState == ECheckBoxState::Checked; })
+								[
+									SNew(STextBlock)
+									.Text(FText::FromString(TEXT("Don't show again in this submit tool instance")))
+								];
+
+							bConfirmed = FDialogFactory::ShowConfirmDialog(Title, MessageBody, AdditionalContent) == EDialogFactoryResult::Confirm;
+						}
+
+						if (bConfirmed)
+						{
+							ModelInterface->ToggleValidator(Node.Pin()->GetValidatorNameId());
+							ParentWidget->ReevaluateValidatorSection();
+						}						
+
+						return FReply::Handled();
+					})
 					[
 						SNew(STextBlock)
 						.Justification(ETextJustify::Center)
diff --git a/Engine/Source/Programs/SubmitTool/Private/View/Widgets/ValidatorsWidget.h b/Engine/Source/Programs/SubmitTool/Private/View/Widgets/ValidatorsWidget.h
--- a/Engine/Source/Programs/SubmitTool/Private/View/Widgets/ValidatorsWidget.h
+++ b/Engine/Source/Programs/SubmitTool/Private/View/Widgets/ValidatorsWidget.h
@@ -39,6 +39,7 @@ public:
 	virtual ~SValidatorsWidget();
 
 	void ReevaluateValidatorSection();
+	bool bDontShowDisablePopup = false;
 private:
 
 	TSharedRef<SWidget> BuildValidatorsView(bool bListPreSubmitOperations = false);
