buildfix Linux: incorrect copy in for loop

diff --git a/Engine/Source/Programs/SubmitTool/Private/Models/ModelInterface.cpp b/Engine/Source/Programs/SubmitTool/Private/Models/ModelInterface.cpp
--- a/Engine/Source/Programs/SubmitTool/Private/Models/ModelInterface.cpp
+++ b/Engine/Source/Programs/SubmitTool/Private/Models/ModelInterface.cpp
@@ -391,7 +391,7 @@ void FModelInterface::EvaluateDisabledValidatorsTag()
 
 
 	TStringBuilder<256> DisabledValidatorsString;
-	for (const TWeakPtr<const FValidatorBase> Validator : ValidationService->GetTasks())
+	for (const TWeakPtr<const FValidatorBase>& Validator : ValidationService->GetTasks())
 	{
 		if (Validator.Pin()->Definition->bIsDisabled)
 		{
