Adding telemetry property to Submit event for disabled validators

diff --git a/Engine/Source/Programs/SubmitTool/Private/Logic/ChangelistService.cpp b/Engine/Source/Programs/SubmitTool/Private/Logic/ChangelistService.cpp
--- a/Engine/Source/Programs/SubmitTool/Private/Logic/ChangelistService.cpp
+++ b/Engine/Source/Programs/SubmitTool/Private/Logic/ChangelistService.cpp
@@ -6,7 +6,6 @@
 #include "Logging/SubmitToolLog.h"
 #include "Logic/Services/Interfaces/ISTSourceControlService.h"
 #include "Models/ModelInterface.h"
-#include "Telemetry/TelemetryService.h"
 #include "CommandLine/CmdLineParameters.h"
 
 #include "ISourceControlProvider.h"
@@ -386,7 +385,6 @@ void FChangelistService::Submit(const FString& InDescriptionAddendum, const FSou
 			{
 				TSharedRef<FCheckIn> CheckIn = StaticCastSharedRef<FCheckIn>(Operation);
 				UE_LOG(LogSubmitToolP4, Log, TEXT("Submit operation succeeded: %s"), *CheckIn->GetSuccessMessage().ToString());
-				FTelemetryService::Get()->SubmitSucceeded();
 			}
 			else if(Result == ECommandResult::Cancelled)
 			{
diff --git a/Engine/Source/Programs/SubmitTool/Private/Models/ModelInterface.cpp b/Engine/Source/Programs/SubmitTool/Private/Models/ModelInterface.cpp
--- a/Engine/Source/Programs/SubmitTool/Private/Models/ModelInterface.cpp
+++ b/Engine/Source/Programs/SubmitTool/Private/Models/ModelInterface.cpp
@@ -9,10 +9,13 @@
 #include "Widgets/Docking/SDockTab.h"
 #include "Containers/Ticker.h"
 #include "Configuration/Configuration.h"
+#include "Policies/CondensedJsonPrintPolicy.h"
+#include "Serialization/JsonWriter.h"
 #include "SourceControlOperations.h"
 #include "HAL/PlatformApplicationMisc.h"
 #include "Async/Async.h"
 #include "Misc/StringOutputDevice.h"
+#include "Telemetry/TelemetryService.h"
 
 #include "Logic/Validators/CoreRedirectsValidator.h"
 #include "Logic/Validators/PreflightValidator.h"
@@ -939,6 +942,27 @@ void FModelInterface::OnSubmitOperationComplete(const FSourceControlOperationRef
 {
 	if (Result == ECommandResult::Succeeded)
 	{
+		FString DisabledTasks;
+		TSharedRef<TJsonWriter<TCHAR, TCondensedJsonPrintPolicy<TCHAR>>> JsonWriter = TJsonWriterFactory<TCHAR, TCondensedJsonPrintPolicy<TCHAR>>::Create(&DisabledTasks, /*Indent=*/0);
+		JsonWriter->WriteArrayStart();
+
+		for (const TWeakPtr<const FValidatorBase>& Task : ValidationService->GetTasks())
+		{
+			if (!Task.Pin()->Definition->bIsDisabled)
+			{
+				JsonWriter->WriteObjectStart();
+				JsonWriter->WriteValue(TEXT("TaskId"), Task.Pin()->GetValidatorNameId().ToString());
+				JsonWriter->WriteObjectEnd();
+			}
+		}
+		JsonWriter->WriteArrayEnd();
+		JsonWriter->Close();
+
+		FTelemetryService::Get()->SubmitSucceeded(
+			MakeAnalyticsEventAttributeArray(
+				TEXT("DisabledTasks"), FJsonFragment(MoveTemp(DisabledTasks))
+			));
+
 		// We've submitted, or tried to submit and failed so we only let the user close the app
 		ChangeState(ESubmitToolAppState::Finished);
 		if(FSubmitToolUserPrefs::Get()->bOpenJiraOnSubmit)
diff --git a/Engine/Source/Programs/SubmitTool/Private/Telemetry/ITelemetry.h b/Engine/Source/Programs/SubmitTool/Private/Telemetry/ITelemetry.h
--- a/Engine/Source/Programs/SubmitTool/Private/Telemetry/ITelemetry.h
+++ b/Engine/Source/Programs/SubmitTool/Private/Telemetry/ITelemetry.h
@@ -11,6 +11,6 @@ class ITelemetry
 public:
 	virtual void Start(const FString& InCurrentStream) const = 0;
 	virtual void BlockFlush(float InTimeout) const = 0;
-	virtual void SubmitSucceeded() const = 0;
+	virtual void SubmitSucceeded(TArray<FAnalyticsEventAttribute>&& InAttribs) const = 0;
 	virtual void CustomEvent(const FString& InEventId, const TArray<FAnalyticsEventAttribute>& InAttribs) const = 0;
 };
\ No newline at end of file
diff --git a/Engine/Source/Programs/SubmitTool/Private/Telemetry/NullTelemetry.h b/Engine/Source/Programs/SubmitTool/Private/Telemetry/NullTelemetry.h
--- a/Engine/Source/Programs/SubmitTool/Private/Telemetry/NullTelemetry.h
+++ b/Engine/Source/Programs/SubmitTool/Private/Telemetry/NullTelemetry.h
@@ -9,6 +9,6 @@ class FNullTelemetry : public ITelemetry
 public:
 	virtual void Start(const FString& InCurrentStream) const override {}
 	virtual void BlockFlush(float InTimeout) const {}
-	virtual void SubmitSucceeded() const override {}
+	virtual void SubmitSucceeded(TArray<FAnalyticsEventAttribute>&& InAttribs) const override {}
 	virtual void CustomEvent(const FString& InEventId, const TArray<FAnalyticsEventAttribute>& InAttribs) const override {}
 };
\ No newline at end of file
diff --git a/Engine/Source/Programs/SubmitTool/Private/Telemetry/StandAloneTelemetry.cpp b/Engine/Source/Programs/SubmitTool/Private/Telemetry/StandAloneTelemetry.cpp
--- a/Engine/Source/Programs/SubmitTool/Private/Telemetry/StandAloneTelemetry.cpp
+++ b/Engine/Source/Programs/SubmitTool/Private/Telemetry/StandAloneTelemetry.cpp
@@ -73,7 +73,7 @@ void FStandAloneTelemetry::CustomEvent(const FString& InEventId, const TArray<FA
 	Provider->RecordEvent(InEventId, InAttribs);
 }
 
-void FStandAloneTelemetry::SubmitSucceeded() const
+void FStandAloneTelemetry::SubmitSucceeded(TArray<FAnalyticsEventAttribute>&& InAttribs) const
 {
 	if (this->Provider == nullptr)
 	{
@@ -82,7 +82,8 @@ void FStandAloneTelemetry::SubmitSucceeded() const
 
 	Provider->RecordEvent(
 		TEXT("SubmitTool.StandAlone.Submit.Succeeded"),
-		MakeAnalyticsEventAttributeArray(
+		AppendAnalyticsEventAttributeArray(
+			InAttribs,
 			TEXT("Version"), FAppVersion::GetVersion()
 		)
 	);
diff --git a/Engine/Source/Programs/SubmitTool/Private/Telemetry/StandAloneTelemetry.h b/Engine/Source/Programs/SubmitTool/Private/Telemetry/StandAloneTelemetry.h
--- a/Engine/Source/Programs/SubmitTool/Private/Telemetry/StandAloneTelemetry.h
+++ b/Engine/Source/Programs/SubmitTool/Private/Telemetry/StandAloneTelemetry.h
@@ -15,7 +15,7 @@ public:
 	virtual ~FStandAloneTelemetry();
 	virtual void Start(const FString& InCurrentStream) const override;
 	virtual void BlockFlush(float InTimeout) const override;
-	virtual void SubmitSucceeded() const override;
+	virtual void SubmitSucceeded(TArray<FAnalyticsEventAttribute>&& InAttribs) const override;
 	virtual void CustomEvent(const FString& InEventId, const TArray<FAnalyticsEventAttribute>& InAttribs) const override;
 
 private:
diff --git a/Engine/Source/Programs/SubmitTool/Private/Telemetry/SubmitToolStudioTelemetry.cpp b/Engine/Source/Programs/SubmitTool/Private/Telemetry/SubmitToolStudioTelemetry.cpp
--- a/Engine/Source/Programs/SubmitTool/Private/Telemetry/SubmitToolStudioTelemetry.cpp
+++ b/Engine/Source/Programs/SubmitTool/Private/Telemetry/SubmitToolStudioTelemetry.cpp
@@ -40,11 +40,12 @@ void FSubmitToolStudioTelemetry::CustomEvent(const FString& InEventId, const TAr
 	FStudioTelemetry::Get().RecordEvent(InEventId, InAttribs);
 }
 
-void FSubmitToolStudioTelemetry::SubmitSucceeded() const
+void FSubmitToolStudioTelemetry::SubmitSucceeded(TArray<FAnalyticsEventAttribute>&& InAttribs) const
 {
 	FStudioTelemetry::Get().RecordEvent(
 		TEXT("SubmitTool.StandAlone.Submit.Succeeded"),
-		MakeAnalyticsEventAttributeArray(
+		AppendAnalyticsEventAttributeArray(
+			InAttribs,
 			TEXT("Version"), FAppVersion::GetVersion()
 		)
 	);
diff --git a/Engine/Source/Programs/SubmitTool/Private/Telemetry/SubmitToolStudioTelemetry.h b/Engine/Source/Programs/SubmitTool/Private/Telemetry/SubmitToolStudioTelemetry.h
--- a/Engine/Source/Programs/SubmitTool/Private/Telemetry/SubmitToolStudioTelemetry.h
+++ b/Engine/Source/Programs/SubmitTool/Private/Telemetry/SubmitToolStudioTelemetry.h
@@ -11,6 +11,6 @@ public:
 	virtual ~FSubmitToolStudioTelemetry();
 	virtual void Start(const FString& InCurrentStream) const override;
 	virtual void BlockFlush(float InTimeout) const override;
-	virtual void SubmitSucceeded() const override;
+	virtual void SubmitSucceeded(TArray<FAnalyticsEventAttribute>&& InAttribs) const override;
 	virtual void CustomEvent(const FString& InEventId, const TArray<FAnalyticsEventAttribute>& InAttribs) const override;
 };
\ No newline at end of file
