Fixed inverse logic in disabled tasks telemetry

#DisabledValidations Tag

diff --git a/Engine/Source/Programs/SubmitTool/Private/Models/ModelInterface.cpp b/Engine/Source/Programs/SubmitTool/Private/Models/ModelInterface.cpp
--- a/Engine/Source/Programs/SubmitTool/Private/Models/ModelInterface.cpp
+++ b/Engine/Source/Programs/SubmitTool/Private/Models/ModelInterface.cpp
@@ -948,7 +948,7 @@ void FModelInterface::OnSubmitOperationComplete(const FSourceControlOperationRef
 
 		for (const TWeakPtr<const FValidatorBase>& Task : ValidationService->GetTasks())
 		{
-			if (!Task.Pin()->Definition->bIsDisabled)
+			if (Task.Pin()->Definition->bIsDisabled)
 			{
 				JsonWriter->WriteObjectStart();
 				JsonWriter->WriteValue(TEXT("TaskId"), Task.Pin()->GetValidatorNameId().ToString());
