Moved disabled validators to its own telemetry event separating it from the submitted event

#DisabledValidations Tag, Preflight, checkCopyright, checkBalancedMacros

diff --git a/Engine/Source/Programs/SubmitTool/Private/Models/ModelInterface.cpp b/Engine/Source/Programs/SubmitTool/Private/Models/ModelInterface.cpp
--- a/Engine/Source/Programs/SubmitTool/Private/Models/ModelInterface.cpp
+++ b/Engine/Source/Programs/SubmitTool/Private/Models/ModelInterface.cpp
@@ -9,8 +9,6 @@
 #include "Widgets/Docking/SDockTab.h"
 #include "Containers/Ticker.h"
 #include "Configuration/Configuration.h"
-#include "Policies/CondensedJsonPrintPolicy.h"
-#include "Serialization/JsonWriter.h"
 #include "SourceControlOperations.h"
 #include "HAL/PlatformApplicationMisc.h"
 #include "Async/Async.h"
@@ -942,26 +940,19 @@ void FModelInterface::OnSubmitOperationComplete(const FSourceControlOperationRef
 {
 	if (Result == ECommandResult::Succeeded)
 	{
-		FString DisabledTasks;
-		TSharedRef<TJsonWriter<TCHAR, TCondensedJsonPrintPolicy<TCHAR>>> JsonWriter = TJsonWriterFactory<TCHAR, TCondensedJsonPrintPolicy<TCHAR>>::Create(&DisabledTasks, /*Indent=*/0);
-		JsonWriter->WriteArrayStart();
-
 		for (const TWeakPtr<const FValidatorBase>& Task : ValidationService->GetTasks())
 		{
 			if (Task.Pin()->Definition->bIsDisabled)
 			{
-				JsonWriter->WriteObjectStart();
-				JsonWriter->WriteValue(TEXT("TaskId"), Task.Pin()->GetValidatorNameId().ToString());
-				JsonWriter->WriteObjectEnd();
+				FTelemetryService::Get()->CustomEvent(TEXT("SubmitTool.StandAlone.Submit.ValidatorDisabled"),
+					MakeAnalyticsEventAttributeArray(
+						TEXT("TaskId"), Task.Pin()->GetValidatorNameId().ToString()
+					)
+				);
 			}
 		}
-		JsonWriter->WriteArrayEnd();
-		JsonWriter->Close();
-
-		FTelemetryService::Get()->SubmitSucceeded(
-			MakeAnalyticsEventAttributeArray(
-				TEXT("DisabledTasks"), FJsonFragment(MoveTemp(DisabledTasks))
-			));
+	
+		FTelemetryService::Get()->SubmitSucceeded(MakeAnalyticsEventAttributeArray());
 
 		// We've submitted, or tried to submit and failed so we only let the user close the app
 		ChangeState(ESubmitToolAppState::Finished);
