Changing telemetry property from validator name id to validator name

#DisabledValidations Tag

diff --git a/Engine/Source/Programs/SubmitTool/Private/Models/ModelInterface.cpp b/Engine/Source/Programs/SubmitTool/Private/Models/ModelInterface.cpp
--- a/Engine/Source/Programs/SubmitTool/Private/Models/ModelInterface.cpp
+++ b/Engine/Source/Programs/SubmitTool/Private/Models/ModelInterface.cpp
@@ -946,7 +946,7 @@ void FModelInterface::OnSubmitOperationComplete(const FSourceControlOperationRef
 			{
 				FTelemetryService::Get()->CustomEvent(TEXT("SubmitTool.StandAlone.Submit.ValidatorDisabled"),
 					MakeAnalyticsEventAttributeArray(
-						TEXT("TaskId"), Task.Pin()->GetValidatorNameId().ToString()
+						TEXT("TaskId"), Task.Pin()->GetValidatorName()
 					)
 				);
 			}
