Fixing a bug in CoreRedirectsValidator where it wouldn't add the last line of a the file downloaded from version control if it didn't end in a newline.

#rnx

diff --git a/Engine/Source/Programs/SubmitTool/Private/Logic/Validators/CoreRedirectsValidator.cpp b/Engine/Source/Programs/SubmitTool/Private/Logic/Validators/CoreRedirectsValidator.cpp
--- a/Engine/Source/Programs/SubmitTool/Private/Logic/Validators/CoreRedirectsValidator.cpp
+++ b/Engine/Source/Programs/SubmitTool/Private/Logic/Validators/CoreRedirectsValidator.cpp
@@ -143,6 +143,12 @@ void FCoreRedirectsValidator::PopulateSetFromStringViewOfFile(FAnsiStringView In
 		OutSet.Add(InFile.Left(LineTerminator));
 		InFile = InFile.RightChop(LineTerminator + 1);
 	}
+	
+	if (InFile.Len())
+	{
+		// Add the last line
+		OutSet.Add(InFile);
+	}
 }
 
 bool FCoreRedirectsValidator::ContainsModifiedRedirects(const TSet<FAnsiStringView>& DepotContents, const TArray<FString>& WorkspaceContents)
