banning <saved by perforce> descriptions

diff --git a/Engine/Source/Programs/SubmitTool/Private/Logic/Validators/TagValidator.cpp b/Engine/Source/Programs/SubmitTool/Private/Logic/Validators/TagValidator.cpp
--- a/Engine/Source/Programs/SubmitTool/Private/Logic/Validators/TagValidator.cpp
+++ b/Engine/Source/Programs/SubmitTool/Private/Logic/Validators/TagValidator.cpp
@@ -12,11 +12,19 @@ bool FTagValidator::Validate(const FString& InCLDescription, const TArray<FSourc
 	bResult = true;
 	bHasFinished = false;
 
-	FString StringToFind = SubmitToolParameters.GeneralParameters.NewChangelistMessage.Left(SubmitToolParameters.GeneralParameters.NewChangelistMessage.Find(TEXT("{")));
-	if(InCLDescription.StartsWith(StringToFind))
+	const TArray<FString> ForbiddenDescriptions
 	{
-		LogFailure(FString::Printf(TEXT("[%s] Please replace the generated description text \"%s\"..."), *GetValidatorName(), *StringToFind));
-		bResult = false;
+		SubmitToolParameters.GeneralParameters.NewChangelistMessage.Left(SubmitToolParameters.GeneralParameters.NewChangelistMessage.Find(TEXT("{"))),
+		TEXT("<saved by perforce>")	
+	};
+
+	for (const FString& StringToFind : ForbiddenDescriptions)
+	{
+		if(InCLDescription.StartsWith(StringToFind, ESearchCase::IgnoreCase))
+		{
+			LogFailure(FString::Printf(TEXT("[%s] Please replace the generated description text \"%s\" providing a useful description..."), *GetValidatorName(), *StringToFind));
+			bResult = false;
+		}
 	}
 
 	TArray<FString> DescriptionLines;
