Added -from-editor cmdline arg to submit tool to indicate it has been launched from editor
Data drive forbidden CL descriptions + trim spaces
Disabled warning popup when enabling a disabled validator
Fixed incompatible files using the root stream to compare filepaths instead of the actual current stream
Added stream to submitted telemetry event
Fixed broken documentation links

diff --git a/Engine/Source/Programs/SubmitTool/Private/CommandLine/CmdLineParameter.h b/Engine/Source/Programs/SubmitTool/Private/CommandLine/CmdLineParameter.h
--- a/Engine/Source/Programs/SubmitTool/Private/CommandLine/CmdLineParameter.h
+++ b/Engine/Source/Programs/SubmitTool/Private/CommandLine/CmdLineParameter.h
@@ -8,17 +8,20 @@ class FCmdLineParameter
 {
 public:
 	FCmdLineParameter() = delete;
-	FCmdLineParameter(FString InKey, bool InIsRequired, FString InDescription, TFunction<bool(const FString& InValue)> InValidator = nullptr, TFunction<void(FString& InValue)> InParseInPlace = nullptr) :
+	FCmdLineParameter(FString InKey, bool InIsRequired, FString InDescription, bool InIsFlag = false, TFunction<bool(const FString& InValue)> InValidator = nullptr, TFunction<void(FString& InValue)> InParseInPlace = nullptr) :
 		Key(InKey),
 		bIsRequired(InIsRequired),
+		bIsFlag(InIsFlag),
 		Description(InDescription),
 		Validator(InValidator),
 		Parser(InParseInPlace)
 	{}
 
-	FString ParameterKey() const { return this->Key; }
-	FString ParameterDescription() const { return this->Description; }
-	bool IsRequired() const { return this->bIsRequired; }
+	FString Key;
+	bool bIsRequired;
+	bool bIsFlag;
+	FString Description;
+
 	bool IsValid(const FString& value) const
 	{
 		if(Validator != nullptr)
@@ -39,9 +42,6 @@ public:
 
 
 private:
-	FString Key;
-	bool bIsRequired;
-	FString Description;
 	TFunction<bool(const FString& InParam)> Validator;
 	TFunction<void(FString& InParam)> Parser;
 };
\ No newline at end of file
diff --git a/Engine/Source/Programs/SubmitTool/Private/CommandLine/CmdLineParameters.cpp b/Engine/Source/Programs/SubmitTool/Private/CommandLine/CmdLineParameters.cpp
--- a/Engine/Source/Programs/SubmitTool/Private/CommandLine/CmdLineParameters.cpp
+++ b/Engine/Source/Programs/SubmitTool/Private/CommandLine/CmdLineParameters.cpp
@@ -18,9 +18,9 @@ bool FCmdLineParameters::ValidateParameters() const
 	bool bIsValid = true;
 	for (const TSharedPtr<FCmdLineParameter>& Parameter : Parameters)
 	{
-		if (Parameter->IsRequired())
+		if (Parameter->bIsRequired)
 		{
-			FString Key = Parameter->ParameterKey();
+			FString Key = Parameter->Key;
 			if(!Key.EndsWith(TEXT(" ")))
 			{
 				Key += TEXT(" ");
@@ -31,18 +31,18 @@ bool FCmdLineParameters::ValidateParameters() const
 			{
 				if (Val.IsEmpty())
 				{
-					UE_LOG(LogSubmitTool, Error, TEXT("Command Line argument '-%s' has no value."), *Parameter->ParameterKey());
+					UE_LOG(LogSubmitTool, Error, TEXT("Command Line argument '-%s' has no value."), *Parameter->Key);
 					bIsValid = false;
 				}
 				else if(!Parameter->IsValid(Val))
 				{
-					UE_LOG(LogSubmitTool, Error, TEXT("Command Line argument '-%s' value '%s' is invalid."), *Parameter->ParameterKey(), *Val);
+					UE_LOG(LogSubmitTool, Error, TEXT("Command Line argument '-%s' value '%s' is invalid."), *Parameter->Key, *Val);
 					bIsValid = false;
 				}
 			}
 			else
 			{
-				UE_LOG(LogSubmitTool, Error, TEXT("Command Line missing '-%s' argument."), *Parameter->ParameterKey());
+				UE_LOG(LogSubmitTool, Error, TEXT("Command Line missing '-%s' argument."), *Parameter->Key);
 				bIsValid = false;
 			}
 		}
@@ -50,11 +50,11 @@ bool FCmdLineParameters::ValidateParameters() const
 		{
 			// optional parameters must contain a value if they are present in the command line
 			FString Val;
-			if (FParse::Value(CommandLine, *Parameter->ParameterKey(), Val))
+			if (!Parameter->bIsFlag && FParse::Value(CommandLine, *Parameter->Key, Val))
 			{
 				if (Val.IsEmpty())
 				{
-					UE_LOG(LogSubmitTool, Error, TEXT("Command Line argument '-%s' has no value."), *Parameter->ParameterKey());
+					UE_LOG(LogSubmitTool, Error, TEXT("Command Line argument '-%s' has no value."), *Parameter->Key);
 					bIsValid = false;
 				}
 			}
@@ -68,7 +68,7 @@ void FCmdLineParameters::LogParameters() const
 {
 	for (const TSharedPtr<FCmdLineParameter>& Parameter : Parameters)
 	{
-		UE_LOG(LogSubmitTool, Warning, TEXT("-%s\t%s"), *Parameter->ParameterKey(), *Parameter->ParameterDescription());
+		UE_LOG(LogSubmitTool, Warning, TEXT("-%s\t%s"), *Parameter->Key, *Parameter->Description);
 	}
 }
 
@@ -89,39 +89,42 @@ bool FCmdLineParameters::GetValue(const FString& InKey, FString& OutValue) const
 	TArray<FString> Tokens;
 	TArray<FString> Switches;
 	FCommandLine::Parse(CommandLine, Tokens, Switches);
-	const TSharedPtr<FCmdLineParameter>* Definition = Parameters.FindByPredicate([InKey](const TSharedPtr<FCmdLineParameter>& InCmdParam) { return InCmdParam->ParameterKey() == InKey; });
+	const TSharedPtr<FCmdLineParameter>* Definition = Parameters.FindByPredicate([InKey](const TSharedPtr<FCmdLineParameter>& InCmdParam) { return InCmdParam->Key == InKey; });
 
-	for(int i = 0; i < Tokens.Num(); i++)
-	{
-		if(Tokens[i].Equals(InKey, ESearchCase::IgnoreCase) && i + 1 < Tokens.Num())
-		{
-			OutValue = Tokens[i + 1];
-			if(Definition)
-			{
-				(*Definition)->CustomParse(OutValue);
-			}
-			return true;
-		}
-	}
-
-	for(int i = 0; i < Switches.Num(); i++)
+	for (int i = 0; i < Switches.Num(); i++)
 	{
 		FString Switch = Switches[i];
 
-		if(Switch.StartsWith(InKey, ESearchCase::IgnoreCase))
+		if (Switch.StartsWith(InKey, ESearchCase::IgnoreCase))
 		{
 			TArray<FString> SplitSwitch;
-			if(2 == Switch.ParseIntoArray(SplitSwitch, TEXT("="), true))
+			if (2 == Switch.ParseIntoArray(SplitSwitch, TEXT("="), true))
 			{
 				OutValue = SplitSwitch[1];
-				if(Definition)
+				if (Definition)
 				{
 					(*Definition)->CustomParse(OutValue);
 				}
+
 				return true;
 			}
 		}
 	}
 
+	for(int i = 0; i < Tokens.Num(); i++)
+	{
+		if(Tokens[i].Equals(InKey, ESearchCase::IgnoreCase) && i + 1 < Tokens.Num())
+		{
+			OutValue = (*Definition)->bIsFlag ? TEXT("true") : Tokens[i + 1];
+			if(Definition)
+			{
+				(*Definition)->CustomParse(OutValue);
+			}
+			return true;
+		}
+	}
+
+
+
 	return false;
 }
diff --git a/Engine/Source/Programs/SubmitTool/Private/CommandLine/CmdLineParameters.h b/Engine/Source/Programs/SubmitTool/Private/CommandLine/CmdLineParameters.h
--- a/Engine/Source/Programs/SubmitTool/Private/CommandLine/CmdLineParameters.h
+++ b/Engine/Source/Programs/SubmitTool/Private/CommandLine/CmdLineParameters.h
@@ -14,14 +14,15 @@ namespace FSubmitToolCmdLine
 	const FString P4ChangeList = TEXT("cl");
 	const FString RootDir = TEXT("root-dir");
 	const FString ParameterFile = TEXT("param-file");
+	const FString EditorFlag = TEXT("from-editor");
 
 	const TArray<TSharedPtr<FCmdLineParameter>> SubmitToolCmdLineArgs =
 	{
 		MakeShared<FCmdLineParameter>(P4Server, true, TEXT("Perforce Server information, expected with format '<address>:<port>'")),
 		MakeShared<FCmdLineParameter>(P4Client, true, TEXT("Perforce workspace name.")),
 		MakeShared<FCmdLineParameter>(P4User, true, TEXT("Perforce user name.")),
-		MakeShared<FCmdLineParameter>(P4ChangeList, true, TEXT("Perforce changelist number to submit."), [](const FString& InValue) { return InValue.IsNumeric() || InValue.Equals(TEXT("default")); }),
-		MakeShared<FCmdLineParameter>(RootDir, false, TEXT("Root directory for the branch this change is part of"), nullptr,
+		MakeShared<FCmdLineParameter>(P4ChangeList, true, TEXT("Perforce changelist number to submit."), false, [](const FString& InValue) { return InValue.IsNumeric() || InValue.Equals(TEXT("default")); }),
+		MakeShared<FCmdLineParameter>(RootDir, false, TEXT("Root directory for the branch this change is part of"), false, nullptr,
 		[](FString& OutValue)
 		{
 			TArray<TCHAR> CharactersToRemove{ '\\', '\"' , '\''};
@@ -58,6 +59,7 @@ namespace FSubmitToolCmdLine
 			FPaths::NormalizeDirectoryName(OutValue);
 		}),
 		MakeShared<FCmdLineParameter>(ParameterFile, false, TEXT("Submit tool parameters XML override.")),
+		MakeShared<FCmdLineParameter>(EditorFlag, false, TEXT("Submit tool parameters XML override."), true),
 	};
 }
 
diff --git a/Engine/Source/Programs/SubmitTool/Private/Logic/Validators/TagValidator.cpp b/Engine/Source/Programs/SubmitTool/Private/Logic/Validators/TagValidator.cpp
--- a/Engine/Source/Programs/SubmitTool/Private/Logic/Validators/TagValidator.cpp
+++ b/Engine/Source/Programs/SubmitTool/Private/Logic/Validators/TagValidator.cpp
@@ -12,15 +12,9 @@ bool FTagValidator::Validate(const FString& InCLDescription, const TArray<FSourc
 	bResult = true;
 	bHasFinished = false;
 
-	const TArray<FString> ForbiddenDescriptions
+	for (const FString& StringToFind : SubmitToolParameters.GeneralParameters.ForbiddenDescriptions)
 	{
-		SubmitToolParameters.GeneralParameters.NewChangelistMessage.Left(SubmitToolParameters.GeneralParameters.NewChangelistMessage.Find(TEXT("{"))),
-		TEXT("<saved by perforce>")	
-	};
-
-	for (const FString& StringToFind : ForbiddenDescriptions)
-	{
-		if(InCLDescription.StartsWith(StringToFind, ESearchCase::IgnoreCase))
+		if(InCLDescription.TrimStartAndEnd().StartsWith(StringToFind, ESearchCase::IgnoreCase))
 		{
 			LogFailure(FString::Printf(TEXT("[%s] Please replace the generated description text \"%s\" providing a useful description..."), *GetValidatorName(), *StringToFind));
 			bResult = false;
diff --git a/Engine/Source/Programs/SubmitTool/Private/Logic/Validators/ValidatorBase.cpp b/Engine/Source/Programs/SubmitTool/Private/Logic/Validators/ValidatorBase.cpp
--- a/Engine/Source/Programs/SubmitTool/Private/Logic/Validators/ValidatorBase.cpp
+++ b/Engine/Source/Programs/SubmitTool/Private/Logic/Validators/ValidatorBase.cpp
@@ -7,6 +7,7 @@
 #include "Misc/Timespan.h"
 #include "Misc/StringOutputDevice.h"
 #include "AnalyticsEventAttribute.h"
+#include "CommandLine/CmdLineParameters.h"
 #include "Configuration/Configuration.h"
 #include "Logic/TagService.h"
 #include "Logic/Services/Interfaces/ICacheDataService.h"
@@ -311,6 +312,14 @@ bool FValidatorBase::EvaluateTagSkip()
 		}
 	}
 
+	if (Definition->bSkipWhenCalledFromEditor && FCmdLineParameters::Get().Contains(FSubmitToolCmdLine::EditorFlag))
+	{
+		UE_LOG(LogValidators, Log, TEXT("[%s] Submit tool has been called from the editor %s doesn't need to run."), *ValidatorName, *ValidatorName);
+		Start = FDateTime::UtcNow();
+		State = EValidationStates::Skipped;
+		return true;
+	}
+
 	if(Definition->bSkipWhenAddendumInDescription && !Definition->ChangelistDescriptionAddendum.IsEmpty())
 	{
 		if(ChangelistService->GetCLDescription().Find(Definition->ChangelistDescriptionAddendum, ESearchCase::IgnoreCase) != INDEX_NONE)
diff --git a/Engine/Source/Programs/SubmitTool/Private/Logic/Validators/ValidatorBase.h b/Engine/Source/Programs/SubmitTool/Private/Logic/Validators/ValidatorBase.h
--- a/Engine/Source/Programs/SubmitTool/Private/Logic/Validators/ValidatorBase.h
+++ b/Engine/Source/Programs/SubmitTool/Private/Logic/Validators/ValidatorBase.h
@@ -21,7 +21,6 @@ struct FAnalyticsEventAttribute;
 
 namespace SubmitToolParseConstants
 {
-	const FString ContentValidator = TEXT("ContentValidator");
 	const FString TagValidator = TEXT("TagValidator");
 	const FString UBTValidator = TEXT("UBTValidator");
 	const FString EditorValidator = TEXT("EditorValidator");
diff --git a/Engine/Source/Programs/SubmitTool/Private/Logic/Validators/ValidatorDefinition.h b/Engine/Source/Programs/SubmitTool/Private/Logic/Validators/ValidatorDefinition.h
--- a/Engine/Source/Programs/SubmitTool/Private/Logic/Validators/ValidatorDefinition.h
+++ b/Engine/Source/Programs/SubmitTool/Private/Logic/Validators/ValidatorDefinition.h
@@ -150,7 +150,13 @@ struct FValidatorDefinition
 	///	Skip this validator when the addendum is already present in the CL description
 	/// 
 	UPROPERTY()
-	bool bSkipWhenAddendumInDescription;
+	bool bSkipWhenAddendumInDescription = false;
+
+	///
+	///	Skip this validator when Submit tool is invoked from editor (has the -from-editor cmdline flag)
+	///
+	UPROPERTY()
+	bool bSkipWhenCalledFromEditor = false;
 
 	///
 	///	Skip is forbidden when this text if found in the CL description
diff --git a/Engine/Source/Programs/SubmitTool/Private/Models/ModelInterface.cpp b/Engine/Source/Programs/SubmitTool/Private/Models/ModelInterface.cpp
--- a/Engine/Source/Programs/SubmitTool/Private/Models/ModelInterface.cpp
+++ b/Engine/Source/Programs/SubmitTool/Private/Models/ModelInterface.cpp
@@ -563,8 +563,6 @@ void FModelInterface::StartSubmitProcess(bool bSkipShelfDialog)
 		{
 			const TArray<FString>& FilesInCL = ChangelistService->GetFilesDepotPaths();
 
-			FConfiguration::AddOrUpdateEntry(TEXT("$(StreamRoot)"), *GetRootStreamName());
-
 			for (const FIncompatibleFilesGroup& FileGroup : Parameters.IncompatibleFilesParams.IncompatibleFileGroups)
 			{
 				TArray<size_t, TInlineAllocator<8>> Indexes;
@@ -573,7 +571,8 @@ void FModelInterface::StartSubmitProcess(bool bSkipShelfDialog)
 				{
 					for (size_t i = 0; i < FileGroup.FileGroups.Num(); ++i)
 					{
-						if (File.Contains(FConfiguration::Substitute(FileGroup.FileGroups[i]), ESearchCase::IgnoreCase))
+						const FString ReplacedPath = FileGroup.FileGroups[i].Replace(TEXT("$(StreamRoot)"), *GetCurrentStream());
+						if (File.Contains(ReplacedPath, ESearchCase::IgnoreCase))
 						{
 							if (!Indexes.Contains(i))
 							{
@@ -952,7 +951,9 @@ void FModelInterface::OnSubmitOperationComplete(const FSourceControlOperationRef
 			}
 		}
 	
-		FTelemetryService::Get()->SubmitSucceeded(MakeAnalyticsEventAttributeArray());
+		FTelemetryService::Get()->SubmitSucceeded(MakeAnalyticsEventAttributeArray(
+			TEXT("Stream"), SourceControlService->GetCurrentStreamName()
+		));
 
 		// We've submitted, or tried to submit and failed so we only let the user close the app
 		ChangeState(ESubmitToolAppState::Finished);
diff --git a/Engine/Source/Programs/SubmitTool/Private/Parameters/SubmitToolParameters.h b/Engine/Source/Programs/SubmitTool/Private/Parameters/SubmitToolParameters.h
--- a/Engine/Source/Programs/SubmitTool/Private/Parameters/SubmitToolParameters.h
+++ b/Engine/Source/Programs/SubmitTool/Private/Parameters/SubmitToolParameters.h
@@ -30,6 +30,13 @@ struct FGeneralParameters
 	UPROPERTY()
 	FString NewChangelistMessage = TEXT("Submit Tool generated changelist from default with {FileCount} files");
 
+	UPROPERTY()
+	TArray<FString> ForbiddenDescriptions = 
+	{
+		TEXT("Submit Tool generated changelist from default with"),
+		TEXT("<saved by perforce>")
+	};
+
 	UPROPERTY()
 	FString CacheFile = TEXT("$(localappdata)/SubmitTool/SubmitToolCache.cache");
 
diff --git a/Engine/Source/Programs/SubmitTool/Private/View/Widgets/ValidatorsWidget.cpp b/Engine/Source/Programs/SubmitTool/Private/View/Widgets/ValidatorsWidget.cpp
--- a/Engine/Source/Programs/SubmitTool/Private/View/Widgets/ValidatorsWidget.cpp
+++ b/Engine/Source/Programs/SubmitTool/Private/View/Widgets/ValidatorsWidget.cpp
@@ -449,6 +449,13 @@ public:
 					.VAlign(EVerticalAlignment::VAlign_Center)
 					.OnClicked_Lambda([this]() 
 					{
+						if (Node.Pin()->Definition->bIsDisabled)
+						{
+							ModelInterface->ToggleValidator(Node.Pin()->GetValidatorNameId());
+							ParentWidget->ReevaluateValidatorSection();
+							return FReply::Handled();
+						}
+
 						bool bConfirmed = ParentWidget->bDontShowDisablePopup;
 
 						if (!bConfirmed)
