Adding option to ban skipping when having certain files in the CL

diff --git a/Engine/Source/Programs/SubmitTool/Private/Logic/Validators/ValidatorBase.cpp b/Engine/Source/Programs/SubmitTool/Private/Logic/Validators/ValidatorBase.cpp
--- a/Engine/Source/Programs/SubmitTool/Private/Logic/Validators/ValidatorBase.cpp
+++ b/Engine/Source/Programs/SubmitTool/Private/Logic/Validators/ValidatorBase.cpp
@@ -299,6 +299,20 @@ void FValidatorBase::ValidationFinished(const bool bHasPassed)
 bool FValidatorBase::EvaluateTagSkip()
 {
 	TSharedPtr<FChangelistService> ChangelistService = ServiceProvider.Pin()->GetService<FChangelistService>();
+	if (Definition->SkipForbiddenFiles.Num() > 0)
+	{
+		for (const FString& File : Definition->SkipForbiddenFiles)
+		{
+			for (const FSourceControlStateRef& SCCFile : ChangelistService->GetFilesInCL())
+			{
+				if (FPaths::GetCleanFilename(SCCFile->GetFilename()).Contains(File))
+				{
+					return false;
+				}
+			}
+		}
+	}
+
 	if(Definition->SkipForbiddenTags.Num() > 0)
 	{
 		for(const FString& Tag : Definition->SkipForbiddenTags)
diff --git a/Engine/Source/Programs/SubmitTool/Private/Logic/Validators/ValidatorDefinition.h b/Engine/Source/Programs/SubmitTool/Private/Logic/Validators/ValidatorDefinition.h
--- a/Engine/Source/Programs/SubmitTool/Private/Logic/Validators/ValidatorDefinition.h
+++ b/Engine/Source/Programs/SubmitTool/Private/Logic/Validators/ValidatorDefinition.h
@@ -164,6 +164,12 @@ struct FValidatorDefinition
 	UPROPERTY()
 	TArray<FString> SkipForbiddenTags;
 
+	///
+	/// Skip is forbidden when the CL contains any of these filepatterns
+	/// 
+	UPROPERTY()
+	TArray<FString> SkipForbiddenFiles;
+
 	///
 	///	Skip is forbidden when this text if found in the CL description
 	/// 
