Adding Packaged Root for referencing bundled paths within a packaged submit tool

diff --git a/Engine/Source/Programs/SubmitTool/Private/Configuration/Configuration.cpp b/Engine/Source/Programs/SubmitTool/Private/Configuration/Configuration.cpp
--- a/Engine/Source/Programs/SubmitTool/Private/Configuration/Configuration.cpp
+++ b/Engine/Source/Programs/SubmitTool/Private/Configuration/Configuration.cpp
@@ -30,6 +30,9 @@ void FConfiguration::Init()
 
 	const FString BatchFileDir = EngineDir + TEXT("/Build/BatchFiles");
 
+	Configuration.Values.Add(TEXT("$(PackagedRoot)"), FPaths::ConvertRelativePathToFull(FPaths::EngineDir() + TEXT("..")));
+	UE_LOG(LogSubmitTool, Log, TEXT("Packaged Root %s"), *FPaths::ConvertRelativePathToFull(FPaths::EngineDir() + TEXT("..")));
+
 #if PLATFORM_WINDOWS
 	Configuration.Values.Add(TEXT("$(RunUAT)"), BatchFileDir + TEXT("/RunUAT.bat"));
 	Configuration.Values.Add(TEXT("$(RunUBT)"), BatchFileDir + TEXT("/RunUBT.bat"));
