Adding tooltip to Run button, changing modifier key from ALT to CTRL

diff --git a/Engine/Source/Programs/SubmitTool/Private/Version/AppVersionDefines.h b/Engine/Source/Programs/SubmitTool/Private/Version/AppVersionDefines.h
--- a/Engine/Source/Programs/SubmitTool/Private/Version/AppVersionDefines.h
+++ b/Engine/Source/Programs/SubmitTool/Private/Version/AppVersionDefines.h
@@ -21,7 +21,7 @@
 
 #define SUBMIT_TOOL_VERSION_MAJOR 6
 #define SUBMIT_TOOL_VERSION_MINOR 2
-#define SUBMIT_TOOL_VERSION_PATCH 0
+#define SUBMIT_TOOL_VERSION_PATCH 1
 
 #ifndef SUBMIT_TOOL_CHANGELIST
 #define SUBMIT_TOOL_CHANGELIST_STRING TEXT("NoCL")
diff --git a/Engine/Source/Programs/SubmitTool/Private/View/Widgets/ValidatorsWidget.cpp b/Engine/Source/Programs/SubmitTool/Private/View/Widgets/ValidatorsWidget.cpp
--- a/Engine/Source/Programs/SubmitTool/Private/View/Widgets/ValidatorsWidget.cpp
+++ b/Engine/Source/Programs/SubmitTool/Private/View/Widgets/ValidatorsWidget.cpp
@@ -403,6 +403,7 @@ public:
 					.IsEnabled_Lambda([this]() {return ModelInterface->GetInputEnabled() && Node.Pin()->GetValidatorState() != EValidationStates::Disabled && bIsRelevantToCL; })
 					.HAlign(EHorizontalAlignment::HAlign_Center)
 					.VAlign(EVerticalAlignment::VAlign_Center)
+					.ToolTipText(FText::FromString(TEXT("Runs this validator. To force run bypassing incremental checks, press down CTRL.")))
 					.OnClicked_Lambda([this]()
 					{
 						if (Node.Pin()->GetIsRunningOrQueued())
@@ -426,7 +427,7 @@ public:
 							}
 							else
 							{
-								if(FSlateApplication::Get().GetModifierKeys().IsLeftAltDown())
+								if(FSlateApplication::Get().GetModifierKeys().IsControlDown())
 								{
 									return FText::FromString("Force Run");
 								}
