Fixing incorrect pathing combination in update service

diff --git a/Engine/Source/Programs/SubmitTool/Private/Logic/UpdateService.cpp b/Engine/Source/Programs/SubmitTool/Private/Logic/UpdateService.cpp
--- a/Engine/Source/Programs/SubmitTool/Private/Logic/UpdateService.cpp
+++ b/Engine/Source/Programs/SubmitTool/Private/Logic/UpdateService.cpp
@@ -72,10 +72,10 @@ void FUpdateService::StartAutoUpdateScript()
 	
 #if PLATFORM_MAC
 	// install path is five levels above the engine path for Mac: ExtractTargetFolder/Mac/SubmitTool.app/Contents/UE/Engine
-	const FString InstallFolder = FPaths::ConvertRelativePathToFull(FString::Printf(TEXT("%s/../../../../.."), *FPaths::EngineDir()));
+	const FString InstallFolder = FPaths::ConvertRelativePathToFull(FString::Printf(TEXT("%s../../../../.."), *FPaths::EngineDir()));
 #else
 	// install path is two levels above the engine path for Linux & Windows: ExtractTargetFolder/Windows/Engine
-	const FString InstallFolder = FPaths::ConvertRelativePathToFull(FString::Printf(TEXT("%s/../.."), *FPaths::EngineDir()));
+	const FString InstallFolder = FPaths::ConvertRelativePathToFull(FString::Printf(TEXT("%s../.."), *FPaths::EngineDir()));
 #endif
 	// copy the auto update script
 	const FString& UpdateScriptDirectory = FPaths::GetPath(this->AutoUpdateParameters.LocalAutoUpdateScript);
