Fix OIDCToken string not being substituted

diff --git a/Engine/Source/Programs/SubmitTool/Private/Configuration/Configuration.cpp b/Engine/Source/Programs/SubmitTool/Private/Configuration/Configuration.cpp
--- a/Engine/Source/Programs/SubmitTool/Private/Configuration/Configuration.cpp
+++ b/Engine/Source/Programs/SubmitTool/Private/Configuration/Configuration.cpp
@@ -31,7 +31,7 @@ void FConfiguration::Init()
 	const FString BatchFileDir = EngineDir + TEXT("/Build/BatchFiles");
 
 	Configuration.Values.Add(TEXT("$(PackagedRoot)"), FPaths::ConvertRelativePathToFull(FPaths::EngineDir() + TEXT("..")));
-	UE_LOG(LogSubmitTool, Log, TEXT("Packaged Root %s"), *FPaths::ConvertRelativePathToFull(FPaths::EngineDir() + TEXT("..")));
+	UE_LOG(LogSubmitToolDebug, Log, TEXT("Packaged Root %s"), *FPaths::ConvertRelativePathToFull(FPaths::EngineDir() + TEXT("..")));
 
 #if PLATFORM_WINDOWS
 	Configuration.Values.Add(TEXT("$(RunUAT)"), BatchFileDir + TEXT("/RunUAT.bat"));
diff --git a/Engine/Source/Programs/SubmitTool/Private/Logic/CredentialsService.cpp b/Engine/Source/Programs/SubmitTool/Private/Logic/CredentialsService.cpp
--- a/Engine/Source/Programs/SubmitTool/Private/Logic/CredentialsService.cpp
+++ b/Engine/Source/Programs/SubmitTool/Private/Logic/CredentialsService.cpp
@@ -171,7 +171,7 @@ void FCredentialsService::GetOIDCToken()
 					ParseOIDCTokenData(FullOutput);
 					});
 
-				OIDCProcess = MakeUnique<FProcessWrapper>(TEXT("Oidc"), Parameters.OAuthTokenTool, FString::Printf(TEXT("%s"), *Parameters.OAuthArgs), OnCompleted, OutputLineProcess);
+				OIDCProcess = MakeUnique<FProcessWrapper>(TEXT("Oidc"), FConfiguration::SubstituteAndNormalizeFilename(Parameters.OAuthTokenTool), FString::Printf(TEXT("%s"), *Parameters.OAuthArgs), OnCompleted, OutputLineProcess);
 				OIDCProcess->Start(true);
 
 				if(!OIDCProcess.IsValid() || OIDCProcess->ExitCode != 0)
