Fixing incorrect regex for warning parse

diff --git a/Engine/Source/Programs/SubmitTool/Private/Logic/Validators/ValidatorDefinition.h b/Engine/Source/Programs/SubmitTool/Private/Logic/Validators/ValidatorDefinition.h
--- a/Engine/Source/Programs/SubmitTool/Private/Logic/Validators/ValidatorDefinition.h
+++ b/Engine/Source/Programs/SubmitTool/Private/Logic/Validators/ValidatorDefinition.h
@@ -330,7 +330,7 @@ struct FValidatorRunExecutableDefinition : public FValidatorDefinition
 	///	Regex for identifying warnings from the output of a process
 	/// 
 	UPROPERTY()
-	FString RegexWarningParsing = TEXT("^(?!.*(?:Display: |Log: ).*( warning |warning:).*$");
+	FString RegexWarningParsing = TEXT("^(?!.*(?:Display: |Log: )).*( warning |warning:).*$");
 
 	TSharedPtr<FRegexPattern> BuiltRegexWarning;
 
