Fixing warning message in incompatible files popup not having the streamname replaced

diff --git a/Engine/Source/Programs/SubmitTool/Private/Models/ModelInterface.cpp b/Engine/Source/Programs/SubmitTool/Private/Models/ModelInterface.cpp
--- a/Engine/Source/Programs/SubmitTool/Private/Models/ModelInterface.cpp
+++ b/Engine/Source/Programs/SubmitTool/Private/Models/ModelInterface.cpp
@@ -586,7 +586,7 @@ void FModelInterface::StartSubmitProcess(bool bSkipShelfDialog)
 				if (Indexes.Num() > 1)
 				{
 					const FText TextTitle = FText::FromString(FileGroup.Title);
-					const FText TextDescription = FText::FromString(FileGroup.GetMessage());
+					const FText TextDescription = FText::FromString(FileGroup.GetMessage().Replace(TEXT("$(StreamRoot)"), *GetCurrentStream()));
 
 					if (FileGroup.bIsError)
 					{
