Fixing tooltip for run button to be more clear

diff --git a/Engine/Source/Programs/SubmitTool/Private/View/Widgets/ValidatorsWidget.cpp b/Engine/Source/Programs/SubmitTool/Private/View/Widgets/ValidatorsWidget.cpp
--- a/Engine/Source/Programs/SubmitTool/Private/View/Widgets/ValidatorsWidget.cpp
+++ b/Engine/Source/Programs/SubmitTool/Private/View/Widgets/ValidatorsWidget.cpp
@@ -403,7 +403,7 @@ public:
 					.IsEnabled_Lambda([this]() {return ModelInterface->GetInputEnabled() && Node.Pin()->GetValidatorState() != EValidationStates::Disabled && bIsRelevantToCL; })
 					.HAlign(EHorizontalAlignment::HAlign_Center)
 					.VAlign(EVerticalAlignment::VAlign_Center)
-					.ToolTipText(FText::FromString(TEXT("Runs this validator. To force run bypassing incremental checks, press down CTRL.")))
+					.ToolTipText(FText::FromString(TEXT("Runs this validator. To force run, hold CTRL.")))
 					.OnClicked_Lambda([this]()
 					{
 						if (Node.Pin()->GetIsRunningOrQueued())
