Changing modifier button from alt to ctrl to align with tooltip info

diff --git a/Engine/Source/Programs/SubmitTool/Private/View/Widgets/ValidatorsWidget.cpp b/Engine/Source/Programs/SubmitTool/Private/View/Widgets/ValidatorsWidget.cpp
--- a/Engine/Source/Programs/SubmitTool/Private/View/Widgets/ValidatorsWidget.cpp
+++ b/Engine/Source/Programs/SubmitTool/Private/View/Widgets/ValidatorsWidget.cpp
@@ -412,7 +412,7 @@ public:
 						}
 						else
 						{
-							ModelInterface->ValidateSingle(Node.Pin()->GetValidatorNameId(), FSlateApplication::Get().GetModifierKeys().IsLeftAltDown());
+							ModelInterface->ValidateSingle(Node.Pin()->GetValidatorNameId(), FSlateApplication::Get().GetModifierKeys().IsControlDown());
 						}
 						return FReply::Handled();
 					})
