Fixing thread safety issue when CL description is updated outside of the UI

#rnx

diff --git a/Engine/Source/Programs/SubmitTool/Private/View/Widgets/SubmitToolWidget.cpp b/Engine/Source/Programs/SubmitTool/Private/View/Widgets/SubmitToolWidget.cpp
--- a/Engine/Source/Programs/SubmitTool/Private/View/Widgets/SubmitToolWidget.cpp
+++ b/Engine/Source/Programs/SubmitTool/Private/View/Widgets/SubmitToolWidget.cpp
@@ -33,6 +33,7 @@
 
 #include "Version/AppVersion.h"
 
+#include "Async/Async.h"
 #include "Modules/ModuleManager.h"
 #include "Framework/Docking/TabManager.h"
 #include "Framework/Application/SlateApplication.h"
@@ -45,7 +46,7 @@
 
 void SubmitToolWidget::OnCLDescriptionUpdated()
 {
-	DescriptionBox->Refresh();
+	AsyncTask(ENamedThreads::GameThread, [this] { DescriptionBox->Refresh(); });	
 }
 
 void SubmitToolWidget::HandleApplicationActivationStateChanged(bool bActive)
