Adding SkipForbiddenStreams to block a validator from skipping in certain streams

diff --git a/Engine/Source/Programs/SubmitTool/Private/Logic/Validators/ValidatorBase.cpp b/Engine/Source/Programs/SubmitTool/Private/Logic/Validators/ValidatorBase.cpp
--- a/Engine/Source/Programs/SubmitTool/Private/Logic/Validators/ValidatorBase.cpp
+++ b/Engine/Source/Programs/SubmitTool/Private/Logic/Validators/ValidatorBase.cpp
@@ -326,6 +326,23 @@ bool FValidatorBase::EvaluateTagSkip()
 		}
 	}
 
+	if (Definition->SkipForbiddenStreams.Num() > 0)
+	{
+		for (const FString& FileDepotPath : ChangelistService->GetFilesDepotPaths())
+		{
+			for (const FString& Stream : Definition->SkipForbiddenStreams)
+			{
+				if (FileDepotPath.StartsWith(Stream, ESearchCase::IgnoreCase))
+				{
+					UE_LOG(LogValidators, Log, TEXT("[%s] Changes in %s require the validation to be run from submit tool. %s won't be skipped"), *ValidatorName, *Stream, *ValidatorName);
+					UE_LOG(LogValidatorsResult, Log, TEXT("[%s] Changes in %s require the validation to be run from submit tool. %s won't be skipped"), *ValidatorName, *Stream, *ValidatorName);
+					return false;
+				}
+			}
+		}
+	}
+
+
 	if (Definition->bSkipWhenCalledFromEditor && FCmdLineParameters::Get().Contains(FSubmitToolCmdLine::EditorFlag))
 	{
 		UE_LOG(LogValidators, Log, TEXT("[%s] Submit tool has been called from the editor %s doesn't need to run."), *ValidatorName, *ValidatorName);
diff --git a/Engine/Source/Programs/SubmitTool/Private/Logic/Validators/ValidatorDefinition.h b/Engine/Source/Programs/SubmitTool/Private/Logic/Validators/ValidatorDefinition.h
--- a/Engine/Source/Programs/SubmitTool/Private/Logic/Validators/ValidatorDefinition.h
+++ b/Engine/Source/Programs/SubmitTool/Private/Logic/Validators/ValidatorDefinition.h
@@ -170,6 +170,12 @@ struct FValidatorDefinition
 	UPROPERTY()
 	TArray<FString> SkipForbiddenFiles;
 
+	///
+	/// Skip is forbidden when the CL contains any file in the specified streams
+	/// 
+	UPROPERTY()
+	TArray<FString> SkipForbiddenStreams;
+
 	///
 	///	Skip is forbidden when this text if found in the CL description
 	/// 
