Adding feature for optional validators where they need to be confirmed as failed to go through with submission

diff --git a/Engine/Source/Programs/SubmitTool/Private/Logic/Validators/PreflightValidator.cpp b/Engine/Source/Programs/SubmitTool/Private/Logic/Validators/PreflightValidator.cpp
--- a/Engine/Source/Programs/SubmitTool/Private/Logic/Validators/PreflightValidator.cpp
+++ b/Engine/Source/Programs/SubmitTool/Private/Logic/Validators/PreflightValidator.cpp
@@ -90,7 +90,61 @@ bool FPreflightValidator::Validate(const FString& InCLDescription, const TArray<
 	return true;
 }
 
-void FPreflightValidator::ValidatePreflights(const TUniquePtr<FPreflightList>& InPreflightListPtr, const TMap<FString, FPreflightData>& InUnlinkedPreflights)
+const FString FPreflightValidator::GetPopupMessageWhenFailedText() const
+{
+	const FString FailedPreflightProp = TEXT("$(FailedPreflightUrls)");
+
+	if(Definition->PopupMessageWhenFailed.Contains(FailedPreflightProp))
+	{
+		TSharedPtr<FPreflightService> PreflightService = ServiceProvider.Pin()->GetService<FPreflightService>();
+		const TArray<FString>& PreflightValues = PreflightTag->GetValues();
+
+		TStringBuilder<1024> PreflightURLs;
+
+		for (FString PreflightId : PreflightValues)
+		{
+			if (PreflightId.Contains(TEXT("/")))
+			{
+				int32 SlashIdx;
+				PreflightId.FindLastChar(TCHAR('/'), SlashIdx);
+				PreflightId = PreflightId.RightChop(SlashIdx + 1);
+			}
+			PreflightId.TrimStartAndEndInline();
+
+			FRegexPattern Pattern = FRegexPattern(TEXT("(?:[0-9]|[a-f]){24}"), ERegexPatternFlags::CaseInsensitive);
+			FRegexMatcher regex = FRegexMatcher(Pattern, PreflightId);
+			if (regex.FindNext())
+			{
+				// Look into the linked preflights and if we don't find it check the unlinked preflights
+				const FPreflightData* FoundData = PreflightService->GetPreflightData()->PreflightList.FindByPredicate([&PreflightId](const FPreflightData& InData) { return InData.ID == PreflightId; });
+				if (FoundData == nullptr)
+				{
+					FoundData = PreflightService->GetUnlinkedPreflights().Find(PreflightId);
+				}
+
+				if(FoundData != nullptr && (FoundData->CachedResults.State != EPreflightState::Completed || FoundData->CachedResults.Outcome != EPreflightOutcome::Success))
+				{
+					if(PreflightURLs.Len() != 0)
+					{
+						PreflightURLs << TEXT(" - ");
+					}
+					PreflightURLs << GetURLTextblock(FoundData->ID);
+				}
+			}
+		}
+
+		return Definition->PopupMessageWhenFailed.Replace(*FailedPreflightProp, PreflightURLs.ToString());
+	}
+
+	return Definition->PopupMessageWhenFailed;
+}
+
+	const FString FPreflightValidator::GetURLTextblock(const FString& InPFId) const
+	{
+		return FString::Printf(TEXT("<a id=\"browser\" style=\"Hyperlink\" href=\"%sjob/%s\">%s</>"), *SubmitToolParameters.HordeParameters.HordeServerAddress, *InPFId, *InPFId);
+	}
+
+	void FPreflightValidator::ValidatePreflights(const TUniquePtr<FPreflightList>& InPreflightListPtr, const TMap<FString, FPreflightData>& InUnlinkedPreflights)
 {
 	if (!InPreflightListPtr.IsValid())
 	{
diff --git a/Engine/Source/Programs/SubmitTool/Private/Logic/Validators/PreflightValidator.h b/Engine/Source/Programs/SubmitTool/Private/Logic/Validators/PreflightValidator.h
--- a/Engine/Source/Programs/SubmitTool/Private/Logic/Validators/PreflightValidator.h
+++ b/Engine/Source/Programs/SubmitTool/Private/Logic/Validators/PreflightValidator.h
@@ -17,7 +17,11 @@ public:
 
 	virtual bool Validate(const FString& InCLDescription, const TArray<FSourceControlStateRef>& InFilteredFilesInCL, const TArray<const FTag*>& InTags) override;
 	virtual const FString& GetValidatorTypeName() const override { return SubmitToolParseConstants::PreflightValidator; }
+
+	virtual const FString GetPopupMessageWhenFailedText() const override;
 protected:
+	const FString GetURLTextblock(const FString& InPFId) const;
+
 	void ValidatePreflights(const TUniquePtr<FPreflightList>& InPreflightList, const TMap<FString, FPreflightData>& InUnlinkedPreflights);
 	virtual void Skip() override;
 
diff --git a/Engine/Source/Programs/SubmitTool/Private/Logic/Validators/ValidatorBase.h b/Engine/Source/Programs/SubmitTool/Private/Logic/Validators/ValidatorBase.h
--- a/Engine/Source/Programs/SubmitTool/Private/Logic/Validators/ValidatorBase.h
+++ b/Engine/Source/Programs/SubmitTool/Private/Logic/Validators/ValidatorBase.h
@@ -166,7 +166,10 @@ public:
 	virtual void InvalidateLocalFileModifications();
 	virtual bool EvaluateTagSkip();
 	virtual bool IsRelevantToCL() const;
-
+	virtual const FString GetPopupMessageWhenFailedText() const
+	{
+		return Definition->PopupMessageWhenFailed;
+	}
 
 	const FString GetStatusText() const;
 	const EValidationStates GetValidatorState() const
diff --git a/Engine/Source/Programs/SubmitTool/Private/Logic/Validators/ValidatorDefinition.h b/Engine/Source/Programs/SubmitTool/Private/Logic/Validators/ValidatorDefinition.h
--- a/Engine/Source/Programs/SubmitTool/Private/Logic/Validators/ValidatorDefinition.h
+++ b/Engine/Source/Programs/SubmitTool/Private/Logic/Validators/ValidatorDefinition.h
@@ -235,6 +235,19 @@ struct FValidatorDefinition
 
 	UPROPERTY()
 	bool bBlocksPreflightStart = true;
+
+	///
+	/// Only usable when optional - If this property is specified and the validator has failed, submit tool will show a 
+	/// pop up with all the failed validator messages that the user has to confirm through.
+	/// 
+	UPROPERTY()
+	FString PopupMessageWhenFailed = FString();
+
+	///
+	///	When enabled this validator failures will be registered in the #FailedValidations tag
+	/// 
+	UPROPERTY()
+	bool bRegisterFailureInCLDescription = false;
 };
 
 USTRUCT()
@@ -339,7 +352,6 @@ struct FValidatorRunExecutableDefinition : public FValidatorDefinition
 	FString RegexWarningParsing = TEXT("^(?!.*(?:Display: |Log: )).*( warning |warning:).*$");
 
 	TSharedPtr<FRegexPattern> BuiltRegexWarning;
-
 };
 
 
diff --git a/Engine/Source/Programs/SubmitTool/Private/Models/ModelInterface.cpp b/Engine/Source/Programs/SubmitTool/Private/Models/ModelInterface.cpp
--- a/Engine/Source/Programs/SubmitTool/Private/Models/ModelInterface.cpp
+++ b/Engine/Source/Programs/SubmitTool/Private/Models/ModelInterface.cpp
@@ -7,6 +7,9 @@
 #include "Framework/Application/SlateApplication.h"
 #include "Version/AppVersion.h"
 #include "Widgets/Docking/SDockTab.h"
+#include "Widgets/Input/SCheckBox.h"
+#include "Widgets/Input/SButton.h"
+#include "Widgets/Text/STextBlock.h"
 #include "Containers/Ticker.h"
 #include "Configuration/Configuration.h"
 #include "SourceControlOperations.h"
@@ -366,54 +369,79 @@ void FModelInterface::EvaluateDisabledValidatorsTag()
 	FString DescriptionCopy = ChangelistService->GetCLDescription();
 	const FString DisabledTag = TEXT("#DisabledValidations ");
 
-	int32 Loc = DescriptionCopy.Find(DisabledTag);
-	if (Loc != INDEX_NONE)
+	bUpdateDescription = RemoveTagFromDescription(DisabledTag, DescriptionCopy);
+
+	TStringBuilder<256> DisabledValidatorsString;
+	for (const TWeakPtr<const FValidatorBase>& Validator : ValidationService->GetTasks())
 	{
-		size_t EndPos = Loc + DisabledTag.Len();
-		while (EndPos < DescriptionCopy.Len())
+		if (Validator.Pin()->Definition->bIsDisabled)
 		{
-			if (DescriptionCopy[EndPos] == TCHAR('\n'))
+			if (DisabledValidatorsString.Len() != 0)
 			{
-				++EndPos;
-				break;
+				DisabledValidatorsString << TEXT(", ");
 			}
-			++EndPos;
-		}
 
-		if (Loc != 0 && DescriptionCopy[Loc - 1] == '\n' && (EndPos - DescriptionCopy.Len()) < 2)
-		{
-			Loc--;
+			DisabledValidatorsString << Validator.Pin()->GetValidatorNameId();
 		}
+	}
 
-		UE_LOG(LogSubmitToolDebug, Log, TEXT("Removed Disabled Validations tag %s"), *DescriptionCopy.Mid(Loc, EndPos - Loc));
-		DescriptionCopy.RemoveAt(Loc, EndPos - Loc);
+	if (DisabledValidatorsString.Len() != 0)
+	{
+		DisabledValidatorsString.InsertAt(0, DisabledTag);
+		const FString FinalString = DisabledValidatorsString.ToString();
+		UE_LOG(LogSubmitToolDebug, Log, TEXT("Added Disabled validators tag %s"), *FinalString);
+		DescriptionCopy.Append(TEXT("\n") + FinalString);
 		bUpdateDescription = true;
 	}
 
 
-	TStringBuilder<256> DisabledValidatorsString;
+	if (bUpdateDescription)
+	{
+		ChangelistService->SetCLDescription(DescriptionCopy, true);
+		TagService->ParseCLDescription();
+	}
+}
+
+
+void FModelInterface::EvaluateFailedValidatorsTag()
+{
+	bool bUpdateDescription = false;
+
+	FString DescriptionCopy = ChangelistService->GetCLDescription();
+	const FString FailedTag = TEXT("#FailedValidations ");
+
+	bUpdateDescription = RemoveTagFromDescription(FailedTag, DescriptionCopy);
+
+	TStringBuilder<256> FailedValidatorsString;
 	for (const TWeakPtr<const FValidatorBase>& Validator : ValidationService->GetTasks())
 	{
-		if (Validator.Pin()->Definition->bIsDisabled)
+		TSharedPtr<const FValidatorBase> ValidatorPtr = Validator.Pin();
+
+		if (ValidatorPtr != nullptr &&
+		ValidatorPtr->Definition->bRegisterFailureInCLDescription &&
+		!ValidatorPtr->GetHasPassed() &&
+		!ValidatorPtr->Definition->bIsDisabled)
 		{
-			if (DisabledValidatorsString.Len() != 0)
+			if (FailedValidatorsString.Len() != 0)
 			{
-				DisabledValidatorsString << TEXT(", ");
+				FailedValidatorsString << TEXT(", ");
 			}
 
-			DisabledValidatorsString << Validator.Pin()->GetValidatorNameId();
+			FailedValidatorsString << Validator.Pin()->GetValidatorNameId();
 		}
 	}
 
-	if (DisabledValidatorsString.Len() != 0)
+	if (FailedValidatorsString.Len() != 0)
 	{
-		DisabledValidatorsString.InsertAt(0, DisabledTag);
-		const FString FinalString = DisabledValidatorsString.ToString();
-		UE_LOG(LogSubmitToolDebug, Log, TEXT("Added Disabled validators tag %s"), *FinalString);
+		FailedValidatorsString.InsertAt(0, FailedTag);
+		const FString FinalString = FailedValidatorsString.ToString();
+		UE_LOG(LogSubmitToolDebug, Log, TEXT("Added Failed validators tag %s"), *FinalString);
 		DescriptionCopy.Append(TEXT("\n") + FinalString);
+		bUpdateDescription = true;
 	}
 
-	if (Loc != INDEX_NONE || DisabledValidatorsString.Len() != 0)
+
+	if (bUpdateDescription)
 	{
 		ChangelistService->SetCLDescription(DescriptionCopy, true);
 		TagService->ParseCLDescription();
@@ -442,28 +470,10 @@ void FModelInterface::UpdateSubmitToolTag(bool InbAdd)
 	}
 	else if (HasSubmitToolTag())
 	{
-		FString VersionlessTag = TEXT("#submittool ");
-		int32 Loc = DescriptionCopy.Find(VersionlessTag);
-		if (Loc >= 0)
-		{
-			size_t EndPos = Loc + VersionlessTag.Len();
-			while (EndPos < DescriptionCopy.Len())
-			{
-				if (DescriptionCopy[EndPos] == TCHAR('\n'))
-				{
-					++EndPos;
-					break;
-				}
-				++EndPos;
-			}
+		const FString VersionlessTag = TEXT("#submittool ");
 
-			if (Loc != 0 && DescriptionCopy[Loc - 1] == '\n' && (EndPos - DescriptionCopy.Len()) < 2)
-			{
-				Loc--;
-			}
-
-			UE_LOG(LogSubmitToolDebug, Log, TEXT("Removed Submit Tool tag"));
-			DescriptionCopy.RemoveAt(Loc, EndPos - Loc);
+		if (RemoveTagFromDescription(VersionlessTag, DescriptionCopy))
+		{
 			ChangelistService->SetCLDescription(DescriptionCopy, true);
 			TagService->ParseCLDescription();
 		}
@@ -556,6 +566,85 @@ void FModelInterface::StartSubmitProcess(bool bSkipShelfDialog)
 			PrepareSubmitCallBack.Broadcast();
 		}
 
+
+		TStringBuilder<1024> SubmitOverFailedValidatorsMessages;
+
+		for (const TWeakPtr<const FValidatorBase>& Validator : ValidationService->GetTasks())
+		{
+			TSharedPtr<const FValidatorBase> ValidatorPtr = Validator.Pin();
+
+			if (ValidatorPtr != nullptr)
+			{
+				// If it was an active validator that failed, and has a popup message set up, build the message
+				if (!ValidatorPtr->Definition->PopupMessageWhenFailed.IsEmpty() &&
+					!ValidatorPtr->Definition->bIsDisabled &&
+					!ValidatorPtr->GetHasPassed())
+				{
+					if (SubmitOverFailedValidatorsMessages.Len() != 0)
+					{
+						SubmitOverFailedValidatorsMessages << TEXT("\n");
+					}
+
+					SubmitOverFailedValidatorsMessages << TEXT("- <RichTextBlock.Red>") << ValidatorPtr->GetValidatorName() << TEXT("</> - ") << ValidatorPtr->GetPopupMessageWhenFailedText() << TEXT("\n");
+				}
+			}
+		}
+
+		if (SubmitOverFailedValidatorsMessages.Len() > 0)
+		{
+			SubmitOverFailedValidatorsMessages.InsertAt(0, TEXT("You have failed important validations, please review each one carefully and confirm you want to continue:\n\n"));
+
+			TArray<FString> Buttons{TEXT("Continue with submission"), TEXT("Cancel")};
+
+			bool bConfirmedSubmission = false;
+
+			TSharedPtr<SHorizontalBox> ConfirmSubmission = SNew(SHorizontalBox)
+			+ SHorizontalBox::Slot()
+			.AutoWidth()
+			.VAlign(VAlign_Center)
+			[
+				SNew(SCheckBox)
+					.IsChecked_Lambda([&bConfirmedSubmission]() { return bConfirmedSubmission ? ECheckBoxState::Checked : ECheckBoxState::Unchecked; })
+					.OnCheckStateChanged_Lambda([&bConfirmedSubmission](ECheckBoxState InNewState)
+						{
+							bConfirmedSubmission = InNewState == ECheckBoxState::Checked;
+						})
+			]
+			+ SHorizontalBox::Slot()
+			.AutoWidth()
+			.VAlign(VAlign_Center)
+			[
+				SNew(SButton)
+					.ButtonStyle(FAppStyle::Get(), "InvisibleButton")
+					.IsFocusable(false)
+					.OnClicked_Lambda([&bConfirmedSubmission]() { bConfirmedSubmission = !bConfirmedSubmission; return FReply::Handled(); })
+					[
+						SNew(STextBlock)
+							.Justification(ETextJustify::Left)
+							.MinDesiredWidth(60)
+							.Text(FText::FromString(TEXT("I confirm I want to submit with these failed validations")))
+					]
+			];
+
+			auto BtnEnabledFunc = [&bConfirmedSubmission](size_t InBtnIdx) 
+			{ 
+				if (InBtnIdx == static_cast<size_t>(EDialogFactoryResult::Confirm))
+				{
+					return bConfirmedSubmission;
+				}
+
+				return true;
+			};
+
+			EDialogFactoryResult DialogResult = FDialogFactory::ShowDialog(FText::FromString(TEXT("Submit with failed validators")), FText::FromString(SubmitOverFailedValidatorsMessages.ToString()), Buttons, ConfirmSubmission, BtnEnabledFunc);
+
+			if (DialogResult != EDialogFactoryResult::Confirm)
+			{
+				return;
+			}
+		}
+
+		EvaluateFailedValidatorsTag();
 		EvaluateDisabledValidatorsTag();
 		UpdateSubmitToolTag(true);
 
@@ -1130,5 +1219,35 @@ void FModelInterface::OnSwarmCreateCompleted(bool InResult, const FString& InErr
 void FModelInterface::ToggleValidator(const FName& ValidatorId)
 {
 	ValidationService->ToggleEnabled(ValidatorId);
-	EvaluateDisabledValidatorsTag();	
+	EvaluateDisabledValidatorsTag();
 }
+
+bool FModelInterface::RemoveTagFromDescription(const FString& InTag, FString& InOutCLDescription) const
+{
+	int32 Loc = InOutCLDescription.Find(InTag);
+	if (Loc != INDEX_NONE)
+	{
+		size_t EndPos = Loc + InTag.Len();
+		while (EndPos < InOutCLDescription.Len())
+		{
+			if (InOutCLDescription[EndPos] == TCHAR('\n'))
+			{
+				++EndPos;
+				break;
+			}
+			++EndPos;
+		}
+
+		if (Loc != 0 && InOutCLDescription[Loc - 1] == '\n' && (EndPos - InOutCLDescription.Len()) < 2)
+		{
+			Loc--;
+		}
+
+		UE_LOG(LogSubmitToolDebug, Log, TEXT("Removed tag %s"), *InOutCLDescription.Mid(Loc, EndPos - Loc));
+		InOutCLDescription.RemoveAt(Loc, EndPos - Loc);
+		return true;
+	}
+
+	return false;
+}
+
diff --git a/Engine/Source/Programs/SubmitTool/Private/Models/ModelInterface.h b/Engine/Source/Programs/SubmitTool/Private/Models/ModelInterface.h
--- a/Engine/Source/Programs/SubmitTool/Private/Models/ModelInterface.h
+++ b/Engine/Source/Programs/SubmitTool/Private/Models/ModelInterface.h
@@ -133,6 +133,7 @@ public:
 	bool IsCLValid() const															{ return ValidationService->GetIsRunSuccessful(!IsIntegrationRequired()); }
 	bool CanLaunchPreflight() const;
 	void EvaluateDisabledValidatorsTag();
+	void EvaluateFailedValidatorsTag();
 	void ReevaluateSubmitToolTag();
 	void UpdateSubmitToolTag(bool InbAdd);
 	bool HasSubmitToolTag() const;
@@ -273,6 +274,7 @@ private:
 	void OnSubmitOperationComplete(const FSourceControlOperationRef& Operation, ECommandResult::Type Result);
 	void OnGetUsersFromSwarmCompleted(const TUniquePtr<FSwarmReview>& InReview, const FString& InErrorMessage);
 	void OnSwarmCreateCompleted(bool InResult, const FString& InErrorMessage);
+	bool RemoveTagFromDescription(const FString& InTag, FString& InOutCLDescription) const;
 
 	FOnChangeListReadyDelegate CLReadyCallback;
 	FOnChangelistRefreshDelegate CLRefreshCallback;
diff --git a/Engine/Source/Programs/SubmitTool/Private/Version/AppVersionDefines.h b/Engine/Source/Programs/SubmitTool/Private/Version/AppVersionDefines.h
--- a/Engine/Source/Programs/SubmitTool/Private/Version/AppVersionDefines.h
+++ b/Engine/Source/Programs/SubmitTool/Private/Version/AppVersionDefines.h
@@ -20,7 +20,7 @@
 #endif
 
 #define SUBMIT_TOOL_VERSION_MAJOR 6
-#define SUBMIT_TOOL_VERSION_MINOR 3
+#define SUBMIT_TOOL_VERSION_MINOR 4
 #define SUBMIT_TOOL_VERSION_PATCH 0
 
 #ifndef SUBMIT_TOOL_CHANGELIST
