Fix executable validator failing to find executable candidates when bValidateExecutableExists set to false

diff --git a/Engine/Source/Programs/SubmitTool/Private/Logic/Validators/ValidatorRunExecutable.cpp b/Engine/Source/Programs/SubmitTool/Private/Logic/Validators/ValidatorRunExecutable.cpp
--- a/Engine/Source/Programs/SubmitTool/Private/Logic/Validators/ValidatorRunExecutable.cpp
+++ b/Engine/Source/Programs/SubmitTool/Private/Logic/Validators/ValidatorRunExecutable.cpp
@@ -110,7 +110,7 @@ bool FValidatorRunExecutable::Activate()
 	const FValidatorRunExecutableDefinition* TypedDefinition = GetTypedDefinition<FValidatorRunExecutableDefinition>();
 
 	FString SelectedOption = OptionsProvider.GetSelectedOptionKey(ExecutableOptions);
-	if(TypedDefinition->ExecutableCandidates.Num() > 0 && SelectedOption.IsEmpty())
+	if(TypedDefinition->bValidateExecutableExists && TypedDefinition->ExecutableCandidates.Num() > 0 && SelectedOption.IsEmpty())
 	{
 		TArray<FString> ExecutablePaths;
 		TypedDefinition->ExecutableCandidates.GenerateValueArray(ExecutablePaths);
@@ -151,7 +151,7 @@ bool FValidatorRunExecutable::Activate()
 			}
 		}
 
-		if(TypedDefinition->bValidateExecutableExists && TypedDefinition->ExecutableCandidates.Num() == 0)
+		if(TypedDefinition->ExecutableCandidates.Num() == 0)
 		{
 			ActivationErrors.Add(FString::Printf(TEXT("Task '%s' does not have a value for 'ExecutablePath' or 'ExecutableCandidates'."), *ValidatorName));
 			bIsValidSetup = false;
