Changed incompatible files text from "Confirm" to "Submit"

diff --git a/Engine/Source/Programs/SubmitTool/Private/Models/ModelInterface.cpp b/Engine/Source/Programs/SubmitTool/Private/Models/ModelInterface.cpp
--- a/Engine/Source/Programs/SubmitTool/Private/Models/ModelInterface.cpp
+++ b/Engine/Source/Programs/SubmitTool/Private/Models/ModelInterface.cpp
@@ -685,7 +685,7 @@ void FModelInterface::StartSubmitProcess(bool bSkipShelfDialog)
 					}
 					else
 					{
-						if (FDialogFactory::ShowConfirmDialog(TextTitle, TextDescription) != EDialogFactoryResult::Confirm)
+						if (FDialogFactory::ShowDialog(TextTitle, TextDescription, TArray<FString>{TEXT("Submit"), TEXT("Cancel")}) != EDialogFactoryResult::FirstButton)
 						{
 							UE_LOG(LogSubmitTool, Log, TEXT("Submission canceled by user"));
 							return;
