Old preflights now won't fail the preflight validator but print a warning instead.
Changed preflight age limit config

diff --git a/Engine/Source/Programs/SubmitTool/Private/Logic/Validators/PreflightValidator.cpp b/Engine/Source/Programs/SubmitTool/Private/Logic/Validators/PreflightValidator.cpp
--- a/Engine/Source/Programs/SubmitTool/Private/Logic/Validators/PreflightValidator.cpp
+++ b/Engine/Source/Programs/SubmitTool/Private/Logic/Validators/PreflightValidator.cpp
@@ -222,29 +222,31 @@ const FString FPreflightValidator::GetPopupMessageWhenFailedText() const
 								LogFailure(FString::Printf(TEXT("[%s] Preflight %s is %d hours old, submitting changes with preflights older than %d hours is not recommended."), *GetValidatorName(), *PreflightId, FMath::FloorToInt(PreflightAge.GetTotalHours()), GetTypedDefinition<FPreflightValidatorDefinition>()->MaxPreflightAgeInHours));
 								WarningPreflights.Add(FoundData);
 							}
-							else
+
+							switch(FoundData->CachedResults.Outcome)
 							{
-								switch(FoundData->CachedResults.Outcome)
-								{
-									case EPreflightOutcome::Success:
-										FoundData->CachedResults.WasSuccessful() ? ValidPreflights.Add(FoundData) : FailedPreflights.Add(FoundData);
-										break;
-									case EPreflightOutcome::Unspecified:
-									case EPreflightOutcome::Failure:
+								case EPreflightOutcome::Success:
+									FoundData->CachedResults.WasSuccessful() ? ValidPreflights.Add(FoundData) : FailedPreflights.Add(FoundData);
+									break;
+								case EPreflightOutcome::Unspecified:
+								case EPreflightOutcome::Failure:
+									FailedPreflights.Add(FoundData);
+									break;
+								case EPreflightOutcome::Warnings:
+									if(Definition->bTreatWarningsAsErrors)
+									{
+										LogFailure(FString::Printf(TEXT("[%s] %s preflight has completed with warnings and they are treated as errors."), *GetValidatorName(), *PreflightId));
 										FailedPreflights.Add(FoundData);
-										break;
-									case EPreflightOutcome::Warnings:
-										if(Definition->bTreatWarningsAsErrors)
-										{
-											LogFailure(FString::Printf(TEXT("[%s] %s preflight has completed with warnings and they are treated as errors."), *GetValidatorName(), *PreflightId));
-											FailedPreflights.Add(FoundData);
-										}
-										else
+									}
+									else
+									{
+										if (!WarningPreflights.Contains(FoundData))
 										{
-											ValidPreflights.Add(FoundData);
+											WarningPreflights.Add(FoundData);
 										}
-										break;
-								}
+										ValidPreflights.Add(FoundData);
+									}
+									break;
 							}
 							break;
 
@@ -299,6 +301,7 @@ const FString FPreflightValidator::GetPopupMessageWhenFailedText() const
 		for (const FPreflightData* PreflightPtr : WarningPreflights)
 		{
 			UE_LOG(LogValidators, Warning, TEXT("[%s] %s preflight has warnings."), *GetValidatorName(), *PreflightPtr->ID);
+			UE_LOG(LogValidatorsResult, Warning, TEXT("[%s] %s preflight has warnings."), *GetValidatorName(), *PreflightPtr->ID);
 		}
 
 		for(const FString& PreflightId : UnverifiedPreflights)
diff --git a/Engine/Source/Programs/SubmitTool/Private/Version/AppVersionDefines.h b/Engine/Source/Programs/SubmitTool/Private/Version/AppVersionDefines.h
--- a/Engine/Source/Programs/SubmitTool/Private/Version/AppVersionDefines.h
+++ b/Engine/Source/Programs/SubmitTool/Private/Version/AppVersionDefines.h
@@ -21,7 +21,7 @@
 
 #define SUBMIT_TOOL_VERSION_MAJOR 6
 #define SUBMIT_TOOL_VERSION_MINOR 4
-#define SUBMIT_TOOL_VERSION_PATCH 0
+#define SUBMIT_TOOL_VERSION_PATCH 1
 
 #ifndef SUBMIT_TOOL_CHANGELIST
 #define SUBMIT_TOOL_CHANGELIST_STRING TEXT("NoCL")
