Added missing '-' for cmdline argument

diff --git a/Engine/Source/Programs/SubmitTool/Private/Logic/Validators/EditorCommandletValidator.cpp b/Engine/Source/Programs/SubmitTool/Private/Logic/Validators/EditorCommandletValidator.cpp
--- a/Engine/Source/Programs/SubmitTool/Private/Logic/Validators/EditorCommandletValidator.cpp
+++ b/Engine/Source/Programs/SubmitTool/Private/Logic/Validators/EditorCommandletValidator.cpp
@@ -248,12 +248,12 @@ FString FEditorCommandletValidator::GetEnableDisablePluginArguments(const TArray
 
 	if (PluginsToDisable.Len() > 0)
 	{
-		CommandlineExtraArgs.Append(FString::Printf(TEXT("DisablePlugins=%s"), PluginsToDisable.ToString()));
+		CommandlineExtraArgs.Append(FString::Printf(TEXT("-DisablePlugins=%s"), PluginsToDisable.ToString()));
 	}
 
 	if (PluginsToEnable.Len() > 0)
 	{
-		CommandlineExtraArgs.Append(FString::Printf(TEXT("EnablePlugins=%s"), PluginsToEnable.ToString()));
+		CommandlineExtraArgs.Append(FString::Printf(TEXT("-EnablePlugins=%s"), PluginsToEnable.ToString()));
 	}
 
 	return CommandlineExtraArgs;
