Adding new bIncludeWarningsInErrorSummary and anchors for enabling/disabling warnings as errors for processes

diff --git a/Engine/Source/Programs/SubmitTool/Private/Logic/Validators/EditorCommandletValidator.cpp b/Engine/Source/Programs/SubmitTool/Private/Logic/Validators/EditorCommandletValidator.cpp
--- a/Engine/Source/Programs/SubmitTool/Private/Logic/Validators/EditorCommandletValidator.cpp
+++ b/Engine/Source/Programs/SubmitTool/Private/Logic/Validators/EditorCommandletValidator.cpp
@@ -253,6 +253,11 @@ FString FEditorCommandletValidator::GetEnableDisablePluginArguments(const TArray
 
 	if (PluginsToEnable.Len() > 0)
 	{
+		if (CommandlineExtraArgs.Len() != 0)
+		{
+			CommandlineExtraArgs.Append(FString::Printf(TEXT(" "), PluginsToDisable.ToString()));
+		}
+
 		CommandlineExtraArgs.Append(FString::Printf(TEXT("-EnablePlugins=%s"), PluginsToEnable.ToString()));
 	}
 
diff --git a/Engine/Source/Programs/SubmitTool/Private/Logic/Validators/ValidatorBase.cpp b/Engine/Source/Programs/SubmitTool/Private/Logic/Validators/ValidatorBase.cpp
--- a/Engine/Source/Programs/SubmitTool/Private/Logic/Validators/ValidatorBase.cpp
+++ b/Engine/Source/Programs/SubmitTool/Private/Logic/Validators/ValidatorBase.cpp
@@ -31,7 +31,7 @@ void FValidatorBase::ParseDefinition(const FString& InDefinition)
 	FStringOutputDevice Errors;
 	FValidatorDefinition::StaticStruct()->ImportText(*InDefinition, const_cast<FValidatorDefinition*>(Definition.Get()), nullptr, 0, &Errors, FValidatorDefinition::StaticStruct()->GetName());
 
-	if(!Errors.IsEmpty())
+	if (!Errors.IsEmpty())
 	{
 		UE_LOG(LogSubmitTool, Error, TEXT("[%s] Error loading parameter file %s"), *GetValidatorNameId().ToString(), *Errors);
 	}
@@ -41,7 +41,7 @@ void FValidatorBase::ParseDefinition(const FString& InDefinition)
 	{
 		FString Extension = PPE.Extension.ToLower();
 		FString Path = FConfiguration::Substitute(PPE.Path);
-		PathsPerExtension.FindOrAdd( MoveTemp(Extension) ).Add( MoveTemp(Path) );
+		PathsPerExtension.FindOrAdd(MoveTemp(Extension)).Add(MoveTemp(Path));
 	}
 }
 
@@ -73,7 +73,7 @@ void FValidatorBase::StartValidation()
 	TArray<FSourceControlStateRef> IncrementallySkippedFiles;
 	bool bAppliesToCL = AppliesToCL(ChangelistService->GetCLDescription(), ChangelistService->GetFilesInCL(), TagService->GetTagsArray(), FilteredFiles, IncrementallySkippedFiles, bIncrementalValidation);
 
-	if(bAppliesToCL)
+	if (bAppliesToCL)
 	{
 		if (FilteredFiles.Num() > 0)
 		{
@@ -110,16 +110,16 @@ void FValidatorBase::StartValidation()
 			}
 		}
 	}
-	else if(!bAppliesToCL)
+	else if (!bAppliesToCL)
 	{
-		if(IncrementallySkippedFiles.Num() != 0)
+		if (IncrementallySkippedFiles.Num() != 0)
 		{
-			UE_LOG(LogValidators, Log, TEXT("[%s] All files were validated in a previous validation and are still valid. To force a validation click 'Run' in the validator list"), *ValidatorName );
+			UE_LOG(LogValidators, Log, TEXT("[%s] All files were validated in a previous validation and are still valid. To force a validation click 'Run' in the validator list"), *ValidatorName);
 			UE_LOG(LogValidatorsResult, Log, TEXT("[%s] All files were validated in a previous validation and are still valid. To force a validation click 'Run' in the validator list"), *ValidatorName);
 		}
 		else
 		{
-			if(!Definition->AppliesToCLRegex.IsEmpty())
+			if (!Definition->AppliesToCLRegex.IsEmpty())
 			{
 				UE_LOG(LogValidators, Log, TEXT("[%s] No files match the regex %s. %s"), *ValidatorName, *FConfiguration::Substitute(Definition->AppliesToCLRegex), *Definition->NotApplicableToCLMessage);
 				UE_LOG(LogValidatorsResult, Log, TEXT("[%s] No files match the regex %s. %s"), *ValidatorName, *FConfiguration::Substitute(Definition->AppliesToCLRegex), *Definition->NotApplicableToCLMessage);
@@ -142,7 +142,7 @@ void FValidatorBase::ToggleEnabled()
 
 	CancelValidation();
 
-	if(ModifiableDefinition->bIsDisabled)
+	if (ModifiableDefinition->bIsDisabled)
 	{
 		State = EValidationStates::Disabled;
 	}
@@ -156,14 +156,14 @@ void FValidatorBase::Tick(float InDeltaTime)
 {
 	RunTime += InDeltaTime;
 
-	if(Definition->TimeoutLimit > 0 && RunTime >= Definition->TimeoutLimit)
+	if (Definition->TimeoutLimit > 0 && RunTime >= Definition->TimeoutLimit)
 	{
 		LogFailure(FString::Printf(TEXT("[%s]: %s"), *GetValidatorName(), TEXT("Timeout limit has been reached, cancelling task.")));
 
 		StopInternalValidations();
 		State = EValidationStates::Timeout;
 
-		if(OnValidationFinished.IsBound())
+		if (OnValidationFinished.IsBound())
 		{
 			OnValidationFinished.Broadcast(*this);
 		}
@@ -174,9 +174,9 @@ bool FValidatorBase::Activate()
 {
 	bIsValidSetup = true;
 
-	if(Definition != nullptr)
+	if (Definition != nullptr)
 	{
-		if(!Definition->IncludeFilesInDirectory.IsEmpty())
+		if (!Definition->IncludeFilesInDirectory.IsEmpty())
 		{
 			FValidatorDefinition* ModifiableDefinition = const_cast<FValidatorDefinition*>(Definition.Get());
 			ModifiableDefinition->IncludeFilesInDirectory = FConfiguration::SubstituteAndNormalizeDirectory(ModifiableDefinition->IncludeFilesInDirectory);
@@ -192,20 +192,20 @@ bool FValidatorBase::Activate()
 
 void FValidatorBase::InvalidateLocalFileModifications()
 {
-	if((Definition->TaskArea & ETaskArea::LocalFiles) == ETaskArea::LocalFiles && (State == EValidationStates::Valid || State == EValidationStates::Running || State == EValidationStates::Skipped || State == EValidationStates::Not_Applicable))
+	if ((Definition->TaskArea & ETaskArea::LocalFiles) == ETaskArea::LocalFiles && (State == EValidationStates::Valid || State == EValidationStates::Running || State == EValidationStates::Skipped || State == EValidationStates::Not_Applicable))
 	{
 		FFileManagerGeneric FileManager;
-		for(const FSourceControlStateRef& File : ServiceProvider.Pin()->GetService<FChangelistService>()->GetFilesInCL())
+		for (const FSourceControlStateRef& File : ServiceProvider.Pin()->GetService<FChangelistService>()->GetFilesInCL())
 		{
 			bool bIncrementallySkipped;
 
-			if(AppliesToFile(File, false, bIncrementallySkipped))
+			if (AppliesToFile(File, false, bIncrementallySkipped))
 			{
 				FString Filename = File->GetFilename();
 				FFileStatData FileModifiedDate = FileManager.GetStatData(*Filename);
-				if(FileModifiedDate.ModificationTime > Start)
+				if (FileModifiedDate.ModificationTime > Start)
 				{
-					if(GetIsRunning())
+					if (GetIsRunning())
 					{
 						UE_LOG(LogValidators, Warning, TEXT("File %s was modified during %s run, this task needs to be run again"), *Filename, *GetValidatorName());
 						UE_LOG(LogValidatorsResult, Warning, TEXT("File %s was modified during %s run, this task needs to be run again"), *Filename, *GetValidatorName());
@@ -229,7 +229,7 @@ const FString FValidatorBase::GetStatusText() const
 	const FString StateStr = StaticEnum<EValidationStates>()->GetNameStringByValue(static_cast<int64>(State))
 		.Replace(TEXT("_"), TEXT(" "));
 
-	if(State == EValidationStates::Skipped || State == EValidationStates::Not_Run || State == EValidationStates::Not_Applicable)
+	if (State == EValidationStates::Skipped || State == EValidationStates::Not_Run || State == EValidationStates::Not_Applicable)
 	{
 		return StateStr;
 	}
@@ -257,21 +257,21 @@ const TArray<FAnalyticsEventAttribute> FValidatorBase::GetTelemetryAttributes()
 		TEXT("Runtime"), RunTime,
 		TEXT("Stream"), ServiceProvider.Pin()->GetService<ISTSourceControlService>()->GetCurrentStreamName(),
 		TEXT("PendingChangelist"), MoveTemp(ChangelistNumber)
-		);
+	);
 }
 
 void FValidatorBase::ValidationFinished(const bool bHasPassed)
 {
-	if(bHasPassed)
+	if (bHasPassed)
 	{
 		UE_LOG(LogValidatorsResult, Log, TEXT("[%s]: Task Succeeded! (%s)"), *GetValidatorName(), *FGenericPlatformTime::PrettyTime(RunTime));
 
-		if(Definition->bUsesIncrementalCache)
+		if (Definition->bUsesIncrementalCache)
 		{
 			ServiceProvider.Pin()->GetService<ICacheDataService>()->UpdateLastValidationForFiles(ServiceProvider.Pin()->GetService<FChangelistService>()->GetCLID(), GetValidatorNameId(), GetValidationConfigId(), FilteredFiles, FDateTime::UtcNow());
 		}
 	}
-	else if(Definition->IsRequired)
+	else if (Definition->IsRequired)
 	{
 		UE_LOG(LogValidatorsResult, Error, TEXT("[%s]: Failed on Required Task!"), *GetValidatorName());
 	}
@@ -280,9 +280,9 @@ void FValidatorBase::ValidationFinished(const bool bHasPassed)
 		UE_LOG(LogValidatorsResult, Warning, TEXT("[%s]: Failed on Optional Task!"), *GetValidatorName());
 	}
 
-	if(!bHasPassed)
+	if (!bHasPassed)
 	{
-		for(const FString& ErrorMsg : Definition->AdditionalValidationErrorMessages)
+		for (const FString& ErrorMsg : Definition->AdditionalValidationErrorMessages)
 		{
 			LogFailure(FString::Printf(TEXT("[%s]: %s"), *GetValidatorName(), *ErrorMsg));
 		}
@@ -290,7 +290,7 @@ void FValidatorBase::ValidationFinished(const bool bHasPassed)
 
 	State = bHasPassed ? EValidationStates::Valid : EValidationStates::Failed;
 
-	if(OnValidationFinished.IsBound())
+	if (OnValidationFinished.IsBound())
 	{
 		OnValidationFinished.Broadcast(*this);
 	}
@@ -313,11 +313,11 @@ bool FValidatorBase::EvaluateTagSkip()
 		}
 	}
 
-	if(Definition->SkipForbiddenTags.Num() > 0)
+	if (Definition->SkipForbiddenTags.Num() > 0)
 	{
-		for(const FString& Tag : Definition->SkipForbiddenTags)
+		for (const FString& Tag : Definition->SkipForbiddenTags)
 		{
-			if(ChangelistService->GetCLDescription().Find(Tag, ESearchCase::IgnoreCase) != INDEX_NONE)
+			if (ChangelistService->GetCLDescription().Find(Tag, ESearchCase::IgnoreCase) != INDEX_NONE)
 			{
 				UE_LOG(LogValidators, Log, TEXT("[%s] The Description contains '%s'. %s is not allowed to be skipped"), *ValidatorName, *Tag, *ValidatorName);
 				UE_LOG(LogValidatorsResult, Log, TEXT("[%s] The Description contains '%s'. %s is not allowed to be skipped"), *ValidatorName, *Tag, *ValidatorName);
@@ -351,9 +351,9 @@ bool FValidatorBase::EvaluateTagSkip()
 		return true;
 	}
 
-	if(Definition->bSkipWhenAddendumInDescription && !Definition->ChangelistDescriptionAddendum.IsEmpty())
+	if (Definition->bSkipWhenAddendumInDescription && !Definition->ChangelistDescriptionAddendum.IsEmpty())
 	{
-		if(ChangelistService->GetCLDescription().Find(Definition->ChangelistDescriptionAddendum, ESearchCase::IgnoreCase) != INDEX_NONE)
+		if (ChangelistService->GetCLDescription().Find(Definition->ChangelistDescriptionAddendum, ESearchCase::IgnoreCase) != INDEX_NONE)
 		{
 			UE_LOG(LogValidators, Log, TEXT("[%s] The Description Addendum '%s' is already present in the CL. %s doesn't need to run"), *ValidatorName, *Definition->ChangelistDescriptionAddendum, *ValidatorName);
 			UE_LOG(LogValidatorsResult, Log, TEXT("[%s] The Description Addendum '%s' is already present in the CL. %s doesn't need to run"), *ValidatorName, *Definition->ChangelistDescriptionAddendum, *ValidatorName);
@@ -396,13 +396,13 @@ bool FValidatorBase::CanPrintErrors() const
 
 void FValidatorBase::PrintErrorSummary() const
 {
-	FScopeLock Lock(&Mutex);
+	FScopeLock Lock(&ErrorLogsMutex);
 
 	if (CanPrintErrors())
 	{
 		if (Definition->IsRequired)
 		{
-			for(const FString& ErrorStr : ErrorListCache)
+			for (const FString& ErrorStr : ErrorListCache)
 			{
 				UE_LOG(LogValidators, Error, TEXT("%s"), *ErrorStr);
 				UE_LOG(LogValidatorsResult, Error, TEXT("%s"), *ErrorStr);
@@ -421,7 +421,7 @@ void FValidatorBase::PrintErrorSummary() const
 const FString FValidatorBase::GetValidationConfigId() const
 {
 	TStringBuilder<512> StringBuilder;
-	for(const TPair<FString, FString>& Pair : OptionsProvider.GetSelectedOptions())
+	for (const TPair<FString, FString>& Pair : OptionsProvider.GetSelectedOptions())
 	{
 		StringBuilder.Append(Pair.Key);
 		StringBuilder.Append(TEXT("_"));
@@ -437,7 +437,7 @@ bool FValidatorBase::AppliesToFile(const FSourceControlStateRef InFile, bool Inb
 	bool bIncluded = false;
 	OutbIsIncrementalSkip = false;
 
-	if((Definition->TaskArea & ETaskArea::LocalFiles) == ETaskArea::None)
+	if ((Definition->TaskArea & ETaskArea::LocalFiles) == ETaskArea::None)
 	{
 		// For validators that do not work on local files, we always apply
 		return true;
@@ -453,99 +453,99 @@ bool FValidatorBase::AppliesToFile(const FSourceControlStateRef InFile, bool Inb
 	}
 	else
 	{
-	if(!InFile->IsDeleted()
-		|| (InFile->IsDeleted() && Definition->bAcceptDeletedFiles))
-	{
-		if (!Definition->AppliesToCLRegex.IsEmpty())
+		if (!InFile->IsDeleted()
+			|| (InFile->IsDeleted() && Definition->bAcceptDeletedFiles))
 		{
-			const FString RegexPat = FConfiguration::Substitute(Definition->AppliesToCLRegex);
-			FRegexPattern Pattern = FRegexPattern(RegexPat, ERegexPatternFlags::CaseInsensitive);
-			FRegexMatcher regex = FRegexMatcher(Pattern, Filename);
+			if (!Definition->AppliesToCLRegex.IsEmpty())
+			{
+				const FString RegexPat = FConfiguration::Substitute(Definition->AppliesToCLRegex);
+				FRegexPattern Pattern = FRegexPattern(RegexPat, ERegexPatternFlags::CaseInsensitive);
+				FRegexMatcher regex = FRegexMatcher(Pattern, Filename);
 
-			bIncluded = regex.FindNext();
-		}
-		else
-		{
-			const bool bIncludeDot = true;
-			FString Extension = FPaths::GetExtension(Filename, bIncludeDot).ToLower();
-			
-			// 1. (If given) the per-extension setting takes over the IncludeFilesInDirectory
-			const TArray<FString>* PathPrefixes = PathsPerExtension.Find(Extension);
-			if (PathPrefixes)
+				bIncluded = regex.FindNext();
+			}
+			else
 			{
-				bool bIncludedInPaths = false;
-				for (const FString& PathPrefix : *PathPrefixes)
+				const bool bIncludeDot = true;
+				FString Extension = FPaths::GetExtension(Filename, bIncludeDot).ToLower();
+
+				// 1. (If given) the per-extension setting takes over the IncludeFilesInDirectory
+				const TArray<FString>* PathPrefixes = PathsPerExtension.Find(Extension);
+				if (PathPrefixes)
 				{
-					if (Filename.StartsWith(PathPrefix, ESearchCase::IgnoreCase))
+					bool bIncludedInPaths = false;
+					for (const FString& PathPrefix : *PathPrefixes)
 					{
-						bIncludedInPaths = true;
-						break;
+						if (Filename.StartsWith(PathPrefix, ESearchCase::IgnoreCase))
+						{
+							bIncludedInPaths = true;
+							break;
+						}
+					}
+					if (!bIncludedInPaths)
+					{
+						return false;
 					}
 				}
-				if (!bIncludedInPaths)
+				// 2. No per-extension setting, use the common setting
+				else if (!Definition->IncludeFilesInDirectory.IsEmpty())
 				{
-					return false;
+					if (!Filename.StartsWith(Definition->IncludeFilesInDirectory, ESearchCase::IgnoreCase))
+					{
+						return false;
+					}
 				}
-			} 
-			// 2. No per-extension setting, use the common setting
-			else if (!Definition->IncludeFilesInDirectory.IsEmpty())
-			{
-				if(!Filename.StartsWith(Definition->IncludeFilesInDirectory, ESearchCase::IgnoreCase))
+				// 3. No directory filter given, carry on
+
+				if (Definition->IncludeFilesWithExtension.IsEmpty())
 				{
-					return false;
+					bIncluded = true;
 				}
-			}
-			// 3. No directory filter given, carry on
-
-			if(Definition->IncludeFilesWithExtension.IsEmpty())
-			{
-				bIncluded = true;
-			}
 
-			for(int Idx = 0; Idx < Definition->IncludeFilesWithExtension.Num(); Idx++)
-			{
-				if(Filename.EndsWith(Definition->IncludeFilesWithExtension[Idx], ESearchCase::IgnoreCase))
+				for (int Idx = 0; Idx < Definition->IncludeFilesWithExtension.Num(); Idx++)
 				{
-					bIncluded = true;
-					break;
+					if (Filename.EndsWith(Definition->IncludeFilesWithExtension[Idx], ESearchCase::IgnoreCase))
+					{
+						bIncluded = true;
+						break;
+					}
 				}
 			}
-		}
 
-		if(bIncluded)
-		{
-			if (!Definition->RequireFileInHierarchy.IsEmpty())
+			if (bIncluded)
 			{
-				if (!FSubmitToolUtils::IsFileInHierarchy(Definition->RequireFileInHierarchy, Filename))
+				if (!Definition->RequireFileInHierarchy.IsEmpty())
 				{
-					bIncluded = false;
+					if (!FSubmitToolUtils::IsFileInHierarchy(Definition->RequireFileInHierarchy, Filename))
+					{
+						bIncluded = false;
+					}
 				}
-			}
 
-			if (!Definition->ExcludeWhenFileInHierarchy.IsEmpty())
-			{
-				if (FSubmitToolUtils::IsFileInHierarchy(Definition->ExcludeWhenFileInHierarchy, Filename))
+				if (!Definition->ExcludeWhenFileInHierarchy.IsEmpty())
 				{
-					bIncluded = false;
+					if (FSubmitToolUtils::IsFileInHierarchy(Definition->ExcludeWhenFileInHierarchy, Filename))
+					{
+						bIncluded = false;
+					}
 				}
 			}
 		}
-		}
 
 		FileHashes.FindOrAdd(FileHash, bIncluded);
 	}
 
-		if(InbAllowIncremental && bIncluded)
-		{
-			FFileManagerGeneric FileManager;
+	if (InbAllowIncremental && bIncluded)
+	{
+		FFileManagerGeneric FileManager;
 		FDateTime LastValidation = ServiceProvider.Pin()->GetService<ICacheDataService>()->GetLastValidationDate(ServiceProvider.Pin()->GetService<FChangelistService>()->GetCLID(), GetValidatorNameId(), GetValidationConfigId(), Filename);
-			FFileStatData FileModifiedDate = FileManager.GetStatData(*Filename);
-			if(LastValidation != FDateTime::MinValue() && FileModifiedDate.ModificationTime < LastValidation)
-			{
-				OutbIsIncrementalSkip = true;
-				bIncluded = false;
-			}
+		FFileStatData FileModifiedDate = FileManager.GetStatData(*Filename);
+		if (LastValidation != FDateTime::MinValue() && FileModifiedDate.ModificationTime < LastValidation)
+		{
+			OutbIsIncrementalSkip = true;
+			bIncluded = false;
 		}
+	}
 
 	return bIncluded;
 }
@@ -558,14 +558,14 @@ bool FValidatorBase::AppliesToCL(const FString& InCLDescription, const TArray<FS
 		return true;
 	}
 
-	for(const FSourceControlStateRef& File : FilesInCL)
+	for (const FSourceControlStateRef& File : FilesInCL)
 	{
 		bool bIsIncrementalSkip;
-		if(AppliesToFile(File, InbAllowIncremental, bIsIncrementalSkip))
+		if (AppliesToFile(File, InbAllowIncremental, bIsIncrementalSkip))
 		{
 			OutFilteredFiles.Add(File);
 		}
-		else if(bIsIncrementalSkip)
+		else if (bIsIncrementalSkip)
 		{
 			OutIncrementalSkips.Add(File);
 		}
diff --git a/Engine/Source/Programs/SubmitTool/Private/Logic/Validators/ValidatorBase.h b/Engine/Source/Programs/SubmitTool/Private/Logic/Validators/ValidatorBase.h
--- a/Engine/Source/Programs/SubmitTool/Private/Logic/Validators/ValidatorBase.h
+++ b/Engine/Source/Programs/SubmitTool/Private/Logic/Validators/ValidatorBase.h
@@ -123,7 +123,7 @@ public:
 		return RunTime;
 	}
 
-	const EValidationStates& GetState()
+	const EValidationStates& GetState() const
 	{
 		return State;
 	}
@@ -253,7 +253,7 @@ protected:
 			UE_LOG(LogValidatorsResult, Warning, TEXT("%s"), *FormattedMsg);
 		}
 
-		FScopeLock Lock(&Mutex);
+		FScopeLock Lock(&ErrorLogsMutex);
 		ErrorListCache.Add(FormattedMsg);
 	}
 
@@ -261,12 +261,12 @@ protected:
 	{
 		State = EValidationStates::Skipped; OnValidationFinished.Broadcast(*this);
 	}
-private:
+
 	EValidationStates State = EValidationStates::Not_Run;
 
 
 	mutable TMap<uint32, bool> FileHashes;
-	mutable FCriticalSection Mutex;
+	mutable FCriticalSection ErrorLogsMutex;
 };
 
 template<DerivedFromDefinition T>
diff --git a/Engine/Source/Programs/SubmitTool/Private/Logic/Validators/ValidatorDefinition.h b/Engine/Source/Programs/SubmitTool/Private/Logic/Validators/ValidatorDefinition.h
--- a/Engine/Source/Programs/SubmitTool/Private/Logic/Validators/ValidatorDefinition.h
+++ b/Engine/Source/Programs/SubmitTool/Private/Logic/Validators/ValidatorDefinition.h
@@ -336,6 +336,24 @@ struct FValidatorRunExecutableDefinition : public FValidatorDefinition
 	/// 
 	UPROPERTY()
 	FString EnableOutputErrorsAnchor;
+	
+	///
+	///	if present, when parsing process output, from this message on, turn off warnings as errors
+	/// 
+	UPROPERTY()
+	FString DisableWarningsAsErrorsAnchor;
+
+	///
+	///	if Present, when parsing process output, from this message on, turn on warnings as errors
+	/// 
+	UPROPERTY()
+	FString EnableWarningsAsErrorsAnchor;
+
+	///
+	/// Whether to include or not the output warnings in the error summary window
+	/// 
+	UPROPERTY()
+	bool bIncludeWarningsInErrorSummary = false;
 
 	///
 	///	Regex for identifying errors from the output of a process
diff --git a/Engine/Source/Programs/SubmitTool/Private/Logic/Validators/ValidatorRunExecutable.cpp b/Engine/Source/Programs/SubmitTool/Private/Logic/Validators/ValidatorRunExecutable.cpp
--- a/Engine/Source/Programs/SubmitTool/Private/Logic/Validators/ValidatorRunExecutable.cpp
+++ b/Engine/Source/Programs/SubmitTool/Private/Logic/Validators/ValidatorRunExecutable.cpp
@@ -202,6 +202,7 @@ bool FValidatorRunExecutable::QueueProcess(const FString& InId, const FString& L
 	
 	FProcessTrackingData ProcessData;
 	ProcessData.bIgnoringOutputError = !GetTypedDefinition<FValidatorRunExecutableDefinition>()->EnableOutputErrorsAnchor.IsEmpty();
+	ProcessData.bIgnoringWarningsAsErrors = !GetTypedDefinition<FValidatorRunExecutableDefinition>()->EnableWarningsAsErrorsAnchor.IsEmpty();
 	ProcessesData.FindOrAdd(InId) = ProcessData;
 
 	FOnOutputLine LineHandler = FOnOutputLine::CreateLambda([this, InId](const FString& InOutput, const EProcessOutputType& InType)
@@ -219,22 +220,30 @@ bool FValidatorRunExecutable::QueueProcess(const FString& InId, const FString& L
 
 void FValidatorRunExecutable::OnProcessOutputLine(const FString& ProcessId, const FString& Line, const EProcessOutputType& OutputType)
 {
-	if(OutputType == EProcessOutputType::ProcessError || IsLineAnError(Line, ProcessesData[ProcessId].bIgnoringOutputError))
+	const FString& FormattedMessage = FString::Printf(TEXT("[%s (%s)]: %s"), *GetValidatorName(), *ProcessId, *Line);
+	if(OutputType == EProcessOutputType::ProcessError || IsLineAnError(Line, ProcessesData[ProcessId].bIgnoringOutputError, ProcessesData[ProcessId].bIgnoringWarningsAsErrors))
 	{
-		const FString& FormattedMessage = FString::Printf(TEXT("[%s (%s)]: %s"), *GetValidatorName(), *ProcessId, *Line);
 		LogFailure(FormattedMessage);
 		ProcessesData[ProcessId].ErrorList.Add(FormattedMessage);
 	}
 	else
 	{
-		UE_LOG(LogValidators, Log, TEXT("[%s (%s)]: %s"), *GetValidatorName(), *ProcessId, *Line);
+		if (GetTypedDefinition<FValidatorRunExecutableDefinition>()->bIncludeWarningsInErrorSummary)
+		{
+			FScopeLock Lock(&ErrorLogsMutex);
+			ProcessesData[ProcessId].ErrorList.Add(FormattedMessage);
+			ErrorListCache.Add(FormattedMessage);
+		}
+		
+		UE_LOG(LogValidators, Log, TEXT("%s"), *FormattedMessage);
 	}
 }
 
-bool FValidatorRunExecutable::IsLineAnError(const FString& InLine, bool& InOutbIgnoringOutputErrors)
+bool FValidatorRunExecutable::IsLineAnError(const FString& InLine, bool& InOutbIgnoringOutputErrors, bool& InOutbIgnoringWarningsAsErrors)
 {
 	const FValidatorRunExecutableDefinition* TypedDefinition = GetTypedDefinition<FValidatorRunExecutableDefinition>();
 
+
 	if(!TypedDefinition->EnableOutputErrorsAnchor.IsEmpty() && InLine.Find(TypedDefinition->EnableOutputErrorsAnchor, ESearchCase::IgnoreCase) != INDEX_NONE)
 	{
 		InOutbIgnoringOutputErrors = false;
@@ -247,6 +256,18 @@ bool FValidatorRunExecutable::IsLineAnError(const FString& InLine, bool& InOutbI
 		return false;
 	}
 
+	if (!TypedDefinition->EnableWarningsAsErrorsAnchor.IsEmpty() && InLine.Find(TypedDefinition->EnableWarningsAsErrorsAnchor, ESearchCase::IgnoreCase) != INDEX_NONE)
+	{
+		InOutbIgnoringWarningsAsErrors = false;
+		return false;
+	}
+
+	if (!TypedDefinition->DisableWarningsAsErrorsAnchor.IsEmpty() && InLine.Find(TypedDefinition->DisableWarningsAsErrorsAnchor, ESearchCase::IgnoreCase) != INDEX_NONE)
+	{
+		InOutbIgnoringWarningsAsErrors = true;
+		return false;
+	}
+
 	if(InOutbIgnoringOutputErrors)
 	{
 		return false;
@@ -269,7 +290,6 @@ bool FValidatorRunExecutable::IsLineAnError(const FString& InLine, bool& InOutbI
 	}
 
 	{
-
 		FRegexMatcher Regex = FRegexMatcher(*TypedDefinition->BuiltRegexError, InLine);
 		if(Regex.FindNext())
 		{
@@ -277,10 +297,10 @@ bool FValidatorRunExecutable::IsLineAnError(const FString& InLine, bool& InOutbI
 		}
 	}
 
-	if(Definition->bTreatWarningsAsErrors)
+	if(Definition->bTreatWarningsAsErrors && !InOutbIgnoringWarningsAsErrors)
 	{
 		FRegexMatcher Regex = FRegexMatcher(*TypedDefinition->BuiltRegexWarning, InLine);
-		if(Regex.FindNext())
+		if (Regex.FindNext())
 		{
 			return true;
 		}
diff --git a/Engine/Source/Programs/SubmitTool/Private/Logic/Validators/ValidatorRunExecutable.h b/Engine/Source/Programs/SubmitTool/Private/Logic/Validators/ValidatorRunExecutable.h
--- a/Engine/Source/Programs/SubmitTool/Private/Logic/Validators/ValidatorRunExecutable.h
+++ b/Engine/Source/Programs/SubmitTool/Private/Logic/Validators/ValidatorRunExecutable.h
@@ -9,6 +9,7 @@ class FBuildVersion;
 struct FProcessTrackingData
 {
 	bool bIgnoringOutputError;
+	bool bIgnoringWarningsAsErrors;
 	TArray<FString> ErrorList;
 };
 
@@ -32,7 +33,7 @@ public:
 protected:
 	virtual bool QueueProcess(const FString& ProcessId, const FString& LocalPath, const FString& Args);
 	virtual void OnProcessOutputLine(const FString& ProcessId, const FString& Line, const EProcessOutputType& OutputType);
-	virtual bool IsLineAnError(const FString& InLine, bool& InOutbIgnoringOutputErrors);
+	virtual bool IsLineAnError(const FString& InLine, bool& InOutbIgnoringOutputErrors, bool& InOutbIgnoringWarningsAsErrors);
 	virtual void OnProcessComplete(const FString& ProcessId, int32 ReturnCode);
 	
 	virtual void Tick(float DeltaTime) override;
diff --git a/Engine/Source/Programs/SubmitTool/Private/Version/AppVersionDefines.h b/Engine/Source/Programs/SubmitTool/Private/Version/AppVersionDefines.h
--- a/Engine/Source/Programs/SubmitTool/Private/Version/AppVersionDefines.h
+++ b/Engine/Source/Programs/SubmitTool/Private/Version/AppVersionDefines.h
@@ -21,7 +21,7 @@
 
 #define SUBMIT_TOOL_VERSION_MAJOR 6
 #define SUBMIT_TOOL_VERSION_MINOR 4
-#define SUBMIT_TOOL_VERSION_PATCH 1
+#define SUBMIT_TOOL_VERSION_PATCH 2
 
 #ifndef SUBMIT_TOOL_CHANGELIST
 #define SUBMIT_TOOL_CHANGELIST_STRING TEXT("NoCL")
