SubmitTool: Optimize build target file to reduce exe size from 40.7 MiB to 33.9 MiB (~100 ms faster boot).
Set bForceDisableAutomationTests=true.
Add WITH_AUTOMATION_WORKER=0.
Add UE_EXTERNAL_PROFILING_ENABLED=0.

#rnx

diff --git a/Engine/Source/Programs/SubmitTool/SubmitTool.Target.cs b/Engine/Source/Programs/SubmitTool/SubmitTool.Target.cs
--- a/Engine/Source/Programs/SubmitTool/SubmitTool.Target.cs
+++ b/Engine/Source/Programs/SubmitTool/SubmitTool.Target.cs
@@ -32,6 +32,7 @@ public class SubmitToolTarget : TargetRules
 		bCompileWithPluginSupport		= true;
 		bBuildRequiresCookedData		= false;
 		bEnableTrace					= true;
+		bForceDisableAutomationTests	= true;
 
 		bHasExports = false;
 		EnablePlugins.Add("PerforceSourceControl");
