SubmitTool: Optimize build target file to reduce exe size from 40.7 MiB to 33.9 MiB (~100 ms faster boot) - Part2.
Set bForceDisableAutomationTests=true.
Add WITH_AUTOMATION_WORKER=0.
Add UE_EXTERNAL_PROFILING_ENABLED=0.

#rnx

diff --git a/Engine/Source/Programs/SubmitTool/SubmitTool.Target.cs b/Engine/Source/Programs/SubmitTool/SubmitTool.Target.cs
--- a/Engine/Source/Programs/SubmitTool/SubmitTool.Target.cs
+++ b/Engine/Source/Programs/SubmitTool/SubmitTool.Target.cs
@@ -44,6 +44,8 @@ public class SubmitToolTarget : TargetRules
 
 		GlobalDefinitions.Add("AUTOSDKS_ENABLED=0");
 		GlobalDefinitions.Add("EXCLUDE_NONPAK_UE_EXTENSIONS=0");
+		GlobalDefinitions.Add("WITH_AUTOMATION_WORKER=0");
+		GlobalDefinitions.Add("UE_EXTERNAL_PROFILING_ENABLED=0");
 
 		OptedInModulePlatforms = new UnrealTargetPlatform[] { UnrealTargetPlatform.Win64, UnrealTargetPlatform.Mac,
 															  UnrealTargetPlatform.Linux, UnrealTargetPlatform.LinuxArm64 };
