SubmitTool: Enable task tracing.

#rnx

diff --git a/Engine/Source/Programs/SubmitTool/SubmitTool.Target.cs b/Engine/Source/Programs/SubmitTool/SubmitTool.Target.cs
--- a/Engine/Source/Programs/SubmitTool/SubmitTool.Target.cs
+++ b/Engine/Source/Programs/SubmitTool/SubmitTool.Target.cs
@@ -42,6 +42,7 @@ public class SubmitToolTarget : TargetRules
 		WindowsPlatform.bUseBundledDbgHelp = false;
 		WindowsPlatform.bPixProfilingEnabled = false;
 
+		GlobalDefinitions.Add("UE_TASK_TRACE_ENABLED=1");
 		GlobalDefinitions.Add("AUTOSDKS_ENABLED=0");
 		GlobalDefinitions.Add("EXCLUDE_NONPAK_UE_EXTENSIONS=0");
 		GlobalDefinitions.Add("WITH_AUTOMATION_WORKER=0");
