SubmitTool: Add some validator trace cpu scopes.

#rnx

diff --git a/Engine/Source/Programs/SubmitTool/Private/Logic/CredentialsService.cpp b/Engine/Source/Programs/SubmitTool/Private/Logic/CredentialsService.cpp
--- a/Engine/Source/Programs/SubmitTool/Private/Logic/CredentialsService.cpp
+++ b/Engine/Source/Programs/SubmitTool/Private/Logic/CredentialsService.cpp
@@ -150,6 +150,8 @@ void FCredentialsService::GetOIDCToken()
 	{
 		GetOIDCTask = UE::Tasks::Launch(UE_SOURCE_LOCATION,
 			[this] {
+				TRACE_CPUPROFILER_EVENT_SCOPE(FCredentialsService::GetOIDCToken);
+
 				FString FullOutput;
 				FOnOutputLine OutputLineProcess = FOnOutputLine::CreateLambda([&FullOutput](const FString& OutputLine, const EProcessOutputType& OutputType) {
 					if(OutputType == EProcessOutputType::ProcessError)
diff --git a/Engine/Source/Programs/SubmitTool/Private/Logic/Validators/PreflightValidator.cpp b/Engine/Source/Programs/SubmitTool/Private/Logic/Validators/PreflightValidator.cpp
--- a/Engine/Source/Programs/SubmitTool/Private/Logic/Validators/PreflightValidator.cpp
+++ b/Engine/Source/Programs/SubmitTool/Private/Logic/Validators/PreflightValidator.cpp
@@ -44,6 +44,7 @@ void FPreflightValidator::ParseDefinition(const FString& InDefinition)
 
 bool FPreflightValidator::Validate(const FString& InCLDescription, const TArray<FSourceControlStateRef>& InFilteredFilesInCL, const TArray<const FTag*>& InTags)
 {
+	TRACE_CPUPROFILER_EVENT_SCOPE(FPreflightValidator::Validate);
 	if(!PreflightTag->IsEnabled() || PreflightTag->GetValues().Num() == 0)
 	{
 		Skip();
