Fix a crash when preflight data couldn't be retrieved due to OIDC token obtention failure

diff --git a/Engine/Source/Programs/SubmitTool/Private/Logic/Validators/PreflightValidator.cpp b/Engine/Source/Programs/SubmitTool/Private/Logic/Validators/PreflightValidator.cpp
--- a/Engine/Source/Programs/SubmitTool/Private/Logic/Validators/PreflightValidator.cpp
+++ b/Engine/Source/Programs/SubmitTool/Private/Logic/Validators/PreflightValidator.cpp
@@ -116,20 +116,24 @@ const FString FPreflightValidator::GetPopupMessageWhenFailedText() const
 			FRegexMatcher regex = FRegexMatcher(Pattern, PreflightId);
 			if (regex.FindNext())
 			{
-				// Look into the linked preflights and if we don't find it check the unlinked preflights
-				const FPreflightData* FoundData = PreflightService->GetPreflightData()->PreflightList.FindByPredicate([&PreflightId](const FPreflightData& InData) { return InData.ID == PreflightId; });
-				if (FoundData == nullptr)
+				const TUniquePtr<FPreflightList>& PFData = PreflightService->GetPreflightData();
+				if (PFData.IsValid())
 				{
-					FoundData = PreflightService->GetUnlinkedPreflights().Find(PreflightId);
-				}
+					// Look into the linked preflights and if we don't find it check the unlinked preflights
+					const FPreflightData* FoundData = PFData->PreflightList.FindByPredicate([&PreflightId](const FPreflightData& InData) { return InData.ID == PreflightId; });
+					if (FoundData == nullptr)
+					{
+						FoundData = PreflightService->GetUnlinkedPreflights().Find(PreflightId);
+					}
 
-				if(FoundData != nullptr && (FoundData->CachedResults.State != EPreflightState::Completed || FoundData->CachedResults.Outcome != EPreflightOutcome::Success))
-				{
-					if(PreflightURLs.Len() != 0)
+					if(FoundData != nullptr && (FoundData->CachedResults.State != EPreflightState::Completed || FoundData->CachedResults.Outcome != EPreflightOutcome::Success))
 					{
-						PreflightURLs << TEXT(" - ");
+						if(PreflightURLs.Len() != 0)
+						{
+							PreflightURLs << TEXT(" - ");
+						}
+						PreflightURLs << GetURLTextblock(FoundData->ID);
 					}
-					PreflightURLs << GetURLTextblock(FoundData->ID);
 				}
 			}
 		}
