Replacing the default CL with the created one when submit tool autoupdates and reopens
Fix Utf8 to Utf16 temporary conversion warning for clang

diff --git a/Engine/Source/Programs/SubmitTool/Private/Logic/UpdateService.cpp b/Engine/Source/Programs/SubmitTool/Private/Logic/UpdateService.cpp
--- a/Engine/Source/Programs/SubmitTool/Private/Logic/UpdateService.cpp
+++ b/Engine/Source/Programs/SubmitTool/Private/Logic/UpdateService.cpp
@@ -14,6 +14,8 @@
 #include "Logic/ProcessWrapper.h"
 #include "Models/HordeDeploymentData.h"
 #include "Framework/Application/SlateApplication.h"
+#include "Logic/Services/SubmitToolServiceProvider.h"
+#include "Logic/ChangelistService.h"
 
 FUpdateService::FUpdateService(const FHordeParameters& InHordeParameters, const FAutoUpdateParameters& InAutoUpdateParameters, TWeakPtr<FSubmitToolServiceProvider> InServiceProvider) :
 	HordeParameters(InHordeParameters),
@@ -427,6 +429,12 @@ FString FUpdateService::GetSubmitToolArgs() const
 	FString Changelist;
 	FCmdLineParameters::Get().GetValue(FSubmitToolCmdLine::P4ChangeList, Changelist);
 
+	// If it's not a valid changelist, try to get it through ChangelistService since we might have used 'default'
+	if (!Changelist.IsNumeric())
+	{
+		Changelist = ServiceProvider.Pin()->GetService<FChangelistService>()->GetCLID();
+	}
+
 	FString PerforceServerAndPort;
 	FCmdLineParameters::Get().GetValue(FSubmitToolCmdLine::P4Server, PerforceServerAndPort);
 
@@ -439,12 +447,16 @@ FString FUpdateService::GetSubmitToolArgs() const
 	FString ParameterFile;
 	FCmdLineParameters::Get().GetValue(FSubmitToolCmdLine::ParameterFile, ParameterFile);
 
-	FString ExecutableArguments = "-server $(server) -user $(user) -client $(client) -cl $(changelist)";
+	FString FromEditor;
+	FCmdLineParameters::Get().GetValue(FSubmitToolCmdLine::EditorFlag, FromEditor);
+
+	FString ExecutableArguments = "-server $(server) -user $(user) -client $(client) -cl $(changelist) $(fromEditor)";
 	ExecutableArguments = ExecutableArguments.Replace(TEXT("$(server)"), *PerforceServerAndPort);
 	ExecutableArguments = ExecutableArguments.Replace(TEXT("$(client)"), *PerforceClientName);
 	ExecutableArguments = ExecutableArguments.Replace(TEXT("$(user)"), *PerforceUserName);
 	ExecutableArguments = ExecutableArguments.Replace(TEXT("$(changelist)"), *Changelist);
 	ExecutableArguments = ExecutableArguments.Replace(TEXT("$(parameterfile)"), *ParameterFile);
+	ExecutableArguments = ExecutableArguments.Replace(TEXT("$(fromEditor)"), *FromEditor);
 
 	return ExecutableArguments;
 }
diff --git a/Engine/Source/Programs/SubmitTool/Private/Parameters/SubmitToolParametersBuilder.cpp b/Engine/Source/Programs/SubmitTool/Private/Parameters/SubmitToolParametersBuilder.cpp
--- a/Engine/Source/Programs/SubmitTool/Private/Parameters/SubmitToolParametersBuilder.cpp
+++ b/Engine/Source/Programs/SubmitTool/Private/Parameters/SubmitToolParametersBuilder.cpp
@@ -81,11 +81,12 @@ FSubmitToolParameters FSubmitToolParametersBuilder::Build()
 	SubmitToolConfig = GConfig->FindConfigFile(IniFilename);
 
 	UE_LOG(LogSubmitTool, Verbose, TEXT("Loading config from the following files:"));
-	for(const TPair<int32, FString>& OverrideLayer : Context.Branch->Hierarchy)
+	for(const TPair<int32, FUtf8String>& OverrideLayer : Context.Branch->Hierarchy)
 	{
-		if(IFileManager::Get().FileExists(*OverrideLayer.Value))
+		FString FullString(OverrideLayer.Value);
+		if(IFileManager::Get().FileExists(*FullString))
 		{
-			UE_LOG(LogSubmitTool, Verbose, TEXT("%s"), *FPaths::ConvertRelativePathToFull(OverrideLayer.Value));
+			UE_LOG(LogSubmitTool, Verbose, TEXT("%s"), *FPaths::ConvertRelativePathToFull(FullString));
 		}
 	}
 
