SubmitTool: Add cpu trace scopes and timing regions for validators.

#rnx

diff --git a/Engine/Source/Programs/SubmitTool/Private/Logic/Validators/ValidatorBase.cpp b/Engine/Source/Programs/SubmitTool/Private/Logic/Validators/ValidatorBase.cpp
--- a/Engine/Source/Programs/SubmitTool/Private/Logic/Validators/ValidatorBase.cpp
+++ b/Engine/Source/Programs/SubmitTool/Private/Logic/Validators/ValidatorBase.cpp
@@ -58,6 +58,9 @@ void FValidatorBase::StartValidation()
 		return;
 	}
 
+	TRACE_BEGIN_REGION(*ValidatorName, TEXT("Validators"));
+	TRACE_CPUPROFILER_EVENT_SCOPE(StartValidation);
+
 	RunTime = 0;
 	Start = FDateTime::UtcNow();
 	State = EValidationStates::Running;
@@ -82,6 +85,7 @@ void FValidatorBase::StartValidation()
 
 		if (!bIsValidSetup)
 		{
+			TRACE_CPUPROFILER_EVENT_SCOPE(Activate);
 			// Try to recover if things have changed since startup
 			ActivationErrors.Empty();
 			if (!Activate())
@@ -98,6 +102,7 @@ void FValidatorBase::StartValidation()
 
 		if (bIsValidSetup)
 		{
+			TRACE_CPUPROFILER_EVENT_SCOPE(Validate);
 			if (IncrementallySkippedFiles.Num() != 0)
 			{
 				const FString FileList = FString::JoinBy(IncrementallySkippedFiles, TEXT("\n"), [](const FSourceControlStateRef& InFile) { return InFile->GetFilename(); });
@@ -112,6 +117,7 @@ void FValidatorBase::StartValidation()
 	}
 	else if (!bAppliesToCL)
 	{
+		TRACE_CPUPROFILER_EVENT_SCOPE(Skip);
 		if (IncrementallySkippedFiles.Num() != 0)
 		{
 			UE_LOG(LogValidators, Log, TEXT("[%s] All files were validated in a previous validation and are still valid. To force a validation click 'Run' in the validator list"), *ValidatorName);
@@ -161,12 +167,7 @@ void FValidatorBase::Tick(float InDeltaTime)
 		LogFailure(FString::Printf(TEXT("[%s]: %s"), *GetValidatorName(), TEXT("Timeout limit has been reached, cancelling task.")));
 
 		StopInternalValidations();
-		State = EValidationStates::Timeout;
-
-		if (OnValidationFinished.IsBound())
-		{
-			OnValidationFinished.Broadcast(*this);
-		}
+		SetValidationFinished(EValidationStates::Timeout);
 	}
 }
 
@@ -288,12 +289,23 @@ void FValidatorBase::ValidationFinished(const bool bHasPassed)
 		}
 	}
 
-	State = bHasPassed ? EValidationStates::Valid : EValidationStates::Failed;
+	EValidationStates NewState = bHasPassed ? EValidationStates::Valid : EValidationStates::Failed;
+	SetValidationFinished(NewState);
+}
+
+void FValidatorBase::SetValidationFinished(EValidationStates NewState)
+{
+	TRACE_BOOKMARK(TEXT("ValidationFinished - %s"), *ValidatorName);
+
+	check(NewState != EValidationStates::Running);
+	check(NewState != EValidationStates::Queued);
+	State = NewState;
 
 	if (OnValidationFinished.IsBound())
 	{
 		OnValidationFinished.Broadcast(*this);
 	}
+	TRACE_END_REGION(*ValidatorName);
 }
 
 bool FValidatorBase::EvaluateTagSkip()
@@ -366,6 +378,11 @@ bool FValidatorBase::EvaluateTagSkip()
 	return false;
 }
 
+void FValidatorBase::Skip()
+{
+	SetValidationFinished(EValidationStates::Skipped);
+}
+
 bool FValidatorBase::IsRelevantToCL() const
 {
 	TSharedPtr<FChangelistService> ChangelistService = ServiceProvider.Pin()->GetService<FChangelistService>();
diff --git a/Engine/Source/Programs/SubmitTool/Private/Logic/Validators/ValidatorBase.h b/Engine/Source/Programs/SubmitTool/Private/Logic/Validators/ValidatorBase.h
--- a/Engine/Source/Programs/SubmitTool/Private/Logic/Validators/ValidatorBase.h
+++ b/Engine/Source/Programs/SubmitTool/Private/Logic/Validators/ValidatorBase.h
@@ -257,16 +257,16 @@ protected:
 		ErrorListCache.Add(FormattedMsg);
 	}
 
-	virtual void Skip()
-	{
-		State = EValidationStates::Skipped; OnValidationFinished.Broadcast(*this);
-	}
+	virtual void Skip();
 
 	EValidationStates State = EValidationStates::Not_Run;
 
 
 	mutable TMap<uint32, bool> FileHashes;
 	mutable FCriticalSection ErrorLogsMutex;
+
+private:
+	void SetValidationFinished(EValidationStates NewState);
 };
 
 template<DerivedFromDefinition T>
