Fixing typos and incorrect log parameters

diff --git a/Engine/Source/Programs/SubmitTool/Private/Logic/Validators/CrossChangelistValidator.cpp b/Engine/Source/Programs/SubmitTool/Private/Logic/Validators/CrossChangelistValidator.cpp
--- a/Engine/Source/Programs/SubmitTool/Private/Logic/Validators/CrossChangelistValidator.cpp
+++ b/Engine/Source/Programs/SubmitTool/Private/Logic/Validators/CrossChangelistValidator.cpp
@@ -172,7 +172,7 @@ bool FCrossChangelistValidator::CheckForFilesInUncontrolledCLFile(const TSet<FSt
 		}
 	}
 
-	return bValid;
+	return bAllValid;
 }
 
 static constexpr const TCHAR* VERSION_NAME = TEXT("version");
diff --git a/Engine/Source/Programs/SubmitTool/Private/Models/ModelInterface.cpp b/Engine/Source/Programs/SubmitTool/Private/Models/ModelInterface.cpp
--- a/Engine/Source/Programs/SubmitTool/Private/Models/ModelInterface.cpp
+++ b/Engine/Source/Programs/SubmitTool/Private/Models/ModelInterface.cpp
@@ -804,7 +804,7 @@ void FModelInterface::RefreshStateBasedOnFiles()
 		{
 			if (!HasSubmitToolTag())
 			{
-				UE_LOG(LogSubmitTool, Warning, TEXT("This CL hasn't been validated and there are no local files. You need to unshelve and run validations."), *ChangelistService->GetCLID());
+				UE_LOG(LogSubmitTool, Warning, TEXT("This CL hasn't been validated and there are no local files. You need to unshelve and run validations."));
 				ChangeState(ESubmitToolAppState::Errored);
 			}
 			else
@@ -1064,7 +1064,7 @@ void FModelInterface::OnSubmitOperationComplete(const FSourceControlOperationRef
 				{
 					if (!JiraValue.Equals(TEXT("none"), ESearchCase::IgnoreCase) && !JiraValue.Equals(TEXT("nojira"), ESearchCase::IgnoreCase))
 					{
-						FString Url = FString::Printf(TEXT("https://%s/browse/%s}"), *Parameters.JiraParameters.ServerAddress, *JiraValue);
+						FString Url = FString::Printf(TEXT("https://%s/browse/%s"), *Parameters.JiraParameters.ServerAddress, *JiraValue);
 						FPlatformProcess::LaunchURL(*Url, nullptr, nullptr);
 					}
 				}
