Added path overrides for property substitutions to submit tool configuration

diff --git a/Engine/Source/Programs/SubmitTool/Private/Parameters/SubmitToolParameters.h b/Engine/Source/Programs/SubmitTool/Private/Parameters/SubmitToolParameters.h
--- a/Engine/Source/Programs/SubmitTool/Private/Parameters/SubmitToolParameters.h
+++ b/Engine/Source/Programs/SubmitTool/Private/Parameters/SubmitToolParameters.h
@@ -54,6 +54,12 @@ struct FGeneralParameters
 
 	UPROPERTY()
 	TArray<FString> GroupsToExclude;
+	
+	UPROPERTY()
+	TMap<FString, FString> PathOverrides;
+
+	UPROPERTY()
+	TMap<FString, FString> PathOverridesInSourceBuild;
 };
 
 USTRUCT()
diff --git a/Engine/Source/Programs/SubmitTool/Private/Parameters/SubmitToolParametersBuilder.cpp b/Engine/Source/Programs/SubmitTool/Private/Parameters/SubmitToolParametersBuilder.cpp
--- a/Engine/Source/Programs/SubmitTool/Private/Parameters/SubmitToolParametersBuilder.cpp
+++ b/Engine/Source/Programs/SubmitTool/Private/Parameters/SubmitToolParametersBuilder.cpp
@@ -124,6 +124,19 @@ FGeneralParameters FSubmitToolParametersBuilder::BuildGeneralParameters()
 		}
 		else
 		{
+			for(const TPair<FString,FString>& Pair : Output.PathOverrides)
+			{				
+				FConfiguration::AddOrUpdateEntry(Pair.Key, Pair.Value);
+			}
+
+			if (!FPaths::IsStaged())
+			{
+				for (const TPair<FString, FString>& Pair : Output.PathOverridesInSourceBuild)
+				{
+					FConfiguration::AddOrUpdateEntry(Pair.Key, Pair.Value);
+				}
+			}
+
 			Output.CacheFile = FConfiguration::SubstituteAndNormalizeFilename(Output.CacheFile);
 		}
 	}
