new validator in submit tool
recursive resolving of path in properties
internal changes

diff --git a/Engine/Source/Programs/SubmitTool/Private/Configuration/Configuration.cpp b/Engine/Source/Programs/SubmitTool/Private/Configuration/Configuration.cpp
--- a/Engine/Source/Programs/SubmitTool/Private/Configuration/Configuration.cpp
+++ b/Engine/Source/Programs/SubmitTool/Private/Configuration/Configuration.cpp
@@ -68,7 +68,7 @@ void FConfiguration::Init()
 	Configuration.Values.Add(TEXT("$(SubmitToolLogsDir)"), *FPaths::ConvertRelativePathToFull(FPaths::ProjectLogDir()));
 
 	Configuration.Values.Add(TEXT("$(SubmitToolConfigDir)"), *FPaths::ConvertRelativePathToFull(FPaths::ProjectConfigDir()));
-
+	
 	Instance = MakeShared<FConfiguration>(Configuration);
 }
 
@@ -102,9 +102,20 @@ FString FConfiguration::Substitute(const FString& InStr)
 
 	FString Output(InStr);
 
-	for (const TPair<FString, FString>& Pair : Instance->Values)
+	bool bReplaced = true;
+	while (bReplaced)
 	{
-		Output.ReplaceInline(*Pair.Key, *Pair.Value, ESearchCase::IgnoreCase);
+		// make sure we do not infinite loop
+		bReplaced = false;
+
+		// check if any entry is replaced
+		for (const TPair<FString, FString>& Pair : Instance->Values)
+		{
+			if (Output.ReplaceInline(*Pair.Key, *Pair.Value, ESearchCase::IgnoreCase) != 0)
+			{
+				bReplaced = true;
+			}
+		}
 	}
 
 	return Output;
