No triggering of validators when submit tool autoupdates

diff --git a/Engine/Source/Programs/SubmitTool/Private/Models/ModelInterface.cpp b/Engine/Source/Programs/SubmitTool/Private/Models/ModelInterface.cpp
--- a/Engine/Source/Programs/SubmitTool/Private/Models/ModelInterface.cpp
+++ b/Engine/Source/Programs/SubmitTool/Private/Models/ModelInterface.cpp
@@ -900,7 +900,10 @@ void FModelInterface::OnChangelistReady(bool bIsValid)
 					}
 				}
 
-				ValidationService->QueueByArea(ValidateArea);
+				if (!UpdateService->CheckForNewVersion())
+				{
+					ValidationService->QueueByArea(ValidateArea);
+				}
 			}
 
 			FileRefreshedCallback.Broadcast();
