SubmitTool: SwarmService: Don't block on the GetSwarmURL p4 property request on startup.

#rnx

diff --git a/Engine/Source/Programs/SubmitTool/Private/Logic/SwarmService.cpp b/Engine/Source/Programs/SubmitTool/Private/Logic/SwarmService.cpp
--- a/Engine/Source/Programs/SubmitTool/Private/Logic/SwarmService.cpp
+++ b/Engine/Source/Programs/SubmitTool/Private/Logic/SwarmService.cpp
@@ -18,7 +18,7 @@
 FSwarmService::FSwarmService(TWeakPtr<FSubmitToolServiceProvider> InServiceProvider) :
 	ServiceProvider(InServiceProvider)
 {
-	GetSwarmURL();
+	LaunchGetSwarmURLTask();
 	ServiceProvider.Pin()->GetService<ISTSourceControlService>()->GetUsersAndGroups(FOnUsersAndGroupsGet::FDelegate::CreateLambda([this](TArray<TSharedPtr<FUserData>>& InUsers, TArray<TSharedPtr<FString>>& InGroups) {
 		Users = &InUsers;
 		Groups = &InGroups;
@@ -293,25 +293,22 @@ void FSwarmService::UpdateReviewDescription(const TDelegate<void(bool)>& OnCompl
 
 }
 
-const FString& FSwarmService::GetSwarmURL()
+void FSwarmService::LaunchGetSwarmURLTask()
 {
-	if(SwarmURL.IsEmpty())
-	{
-		ServiceProvider.Pin()->GetService<ISTSourceControlService>()->RunCommand(TEXT("property"), { "-l", "-n", "P4.Swarm.URL" }, FOnSCCCommandCompleteNoRet::CreateLambda(
-			[this](bool bSuccess, const FSCCRecordSet& InResultValues, const FSourceControlResultInfo& InResultsInfo)
+	GetSwarmURLTask = ServiceProvider.Pin()->GetService<ISTSourceControlService>()->RunCommand(TEXT("property"), { "-l", "-n", "P4.Swarm.URL" },
+		FOnSCCCommandCompleteNoRet::CreateLambda(
+		[this](bool bSuccess, const FSCCRecordSet& InResultValues, const FSourceControlResultInfo& InResultsInfo)
+		{
+			if(bSuccess && InResultValues.Num() > 0 && InResultValues[0].Contains(TEXT("value")))
 			{
-				if(bSuccess && InResultValues.Num() > 0 && InResultValues[0].Contains(TEXT("value")))
-				{
-					SwarmURL = InResultValues[0][TEXT("value")];
-				}
-
-			})).Wait();
-	}
-	return SwarmURL;
+				SwarmURL = InResultValues[0][TEXT("value")];
+			}
+		}));
 }
 
 FString FSwarmService::ReviewsURL() const
 {
+	GetSwarmURLTask.Wait();
 	return SwarmURL.IsEmpty() ? FString{} : SwarmURL / TEXT("api/v9/reviews");
 }
 
diff --git a/Engine/Source/Programs/SubmitTool/Private/Logic/SwarmService.h b/Engine/Source/Programs/SubmitTool/Private/Logic/SwarmService.h
--- a/Engine/Source/Programs/SubmitTool/Private/Logic/SwarmService.h
+++ b/Engine/Source/Programs/SubmitTool/Private/Logic/SwarmService.h
@@ -59,7 +59,7 @@ public:
 	const TArray<FString> GetUsersInSwarmTag() const;
 
 private:
-	const FString& GetSwarmURL();
+	void LaunchGetSwarmURLTask();
 	FString ReviewsURL() const;
 
 
@@ -68,6 +68,7 @@ private:
 	TSharedPtr<IHttpRequest> UpdateSwarmRequest;
 
 	TWeakPtr<FSubmitToolServiceProvider> ServiceProvider;
+	UE::Tasks::FTask GetSwarmURLTask;
 	FString SwarmURL;
 	TUniquePtr<FSwarmReview> Review = nullptr;
 
@@ -76,4 +77,4 @@ private:
 	TArray<TSharedPtr<FString>>* Groups;
 };
 
-Expose_TNameOf(FSwarmService);
\ No newline at end of file
+Expose_TNameOf(FSwarmService);
