SubmitTool: P4LockDownService: Simplify launching and waiting for lock config files downloads.
Use N tasks instead of 2N+1 tasks for N downloads.

#rnx

diff --git a/Engine/Source/Programs/SubmitTool/Private/Logic/P4LockdownService.cpp b/Engine/Source/Programs/SubmitTool/Private/Logic/P4LockdownService.cpp
--- a/Engine/Source/Programs/SubmitTool/Private/Logic/P4LockdownService.cpp
+++ b/Engine/Source/Programs/SubmitTool/Private/Logic/P4LockdownService.cpp
@@ -13,52 +13,41 @@
 #include "Configuration/Configuration.h"
 
 FP4LockdownService::FP4LockdownService(const FP4LockdownParameters& InParameters, TWeakPtr<FSubmitToolServiceProvider> InServiceProvider)
-: 	ServiceProvider(InServiceProvider), Parameters(InParameters)
+	: ConfigFilesTaskEvent(TEXT("ConfigFilesTaskEvent"))
+	, ServiceProvider(InServiceProvider)
+	, Parameters(InParameters)
 {
 	using namespace UE::Tasks;
-	ConfigFilesTask = Launch(UE_SOURCE_LOCATION, [this]
+	for (const TPair<FString, FString>& Pair : Parameters.ConfigPaths)
+	{
+		ConfigFilesTaskEvent.AddPrerequisites(Launch(UE_SOURCE_LOCATION, [this, &Pair]
 		{
-			TArray<TTask<bool>> DownloadFileTasks;
-			for(const TPair<FString, FString>& Pair : Parameters.ConfigPaths)
-			{
-				DownloadFileTasks.Add(Launch(UE_SOURCE_LOCATION, [this, &Pair]
-					{
-						return GetConfigFile(Pair.Key, Pair.Value);
-					}));
-			}
-
-			Wait(DownloadFileTasks);
-
-			bool bResult = true;
-			for(TTask<bool>& Task : DownloadFileTasks)
-			{
-				bResult |= Task.GetResult();
-			}
-
-			return bResult;
-		});
+			return GetConfigFile(Pair.Key, Pair.Value);
+		}));
+	}
+	ConfigFilesTaskEvent.Trigger();
 }
 
 bool FP4LockdownService::ArePathsInLockdown(const TArray<FString>& InPaths, bool& bOutAllowlisted)
 {	
-	if(!ConfigFilesTask.IsValid())
+	if(!ConfigFilesTaskEvent.IsValid())
 	{
 		UE_LOG(LogSubmitToolP4, Error, TEXT("Downloading task wasn't setup correctly, hardlock status is not known."));
 		return false;
 	}
 
-	if(!ConfigFilesTask.IsCompleted())
+	if(!ConfigFilesTaskEvent.IsCompleted())
 	{		
 		UE_LOG(LogSubmitToolP4, Log, TEXT("Waiting for download of Stream Hardlock data..."));
-		ConfigFilesTask.Wait(FTimespan::FromSeconds(5));
+		ConfigFilesTaskEvent.Wait(FTimespan::FromSeconds(5));
 
-		if(!ConfigFilesTask.IsCompleted())
+		if(!ConfigFilesTaskEvent.IsCompleted())
 		{
 			UE_LOG(LogSubmitToolP4, Error, TEXT("Downloading config files from P4 timed out, hardlock status is not latest, will revert to use cache."));
 		}
 	}
 
-	if(ConfigFilesTask.IsCompleted())
+	if(ConfigFilesTaskEvent.IsCompleted())
 	{
 		for(const TPair<FString, FSharedBuffer>& Pair : DownloadedFiles)
 		{
@@ -78,11 +67,6 @@ bool FP4LockdownService::ArePathsInLockdown(const TArray<FString>& InPaths, bool
 		}
 	}
 
-	if(!ConfigFilesTask.GetResult())
-	{
-		UE_LOG(LogSubmitToolP4, Error, TEXT("Downloading config files from P4 failed. Lockdown data won't be available"));
-	}
-
 	ParseAllowListData();
 
 	if(InPaths.IsEmpty())
diff --git a/Engine/Source/Programs/SubmitTool/Private/Logic/P4LockdownService.h b/Engine/Source/Programs/SubmitTool/Private/Logic/P4LockdownService.h
--- a/Engine/Source/Programs/SubmitTool/Private/Logic/P4LockdownService.h
+++ b/Engine/Source/Programs/SubmitTool/Private/Logic/P4LockdownService.h
@@ -21,7 +21,7 @@ public:
 	bool ArePathsInLockdown(const TArray<FString>& InPaths, bool& bOutAllowlisted);
 	bool IsBlockingOperationRunning() const
 	{
-		return ConfigFilesTask.IsValid() && !ConfigFilesTask.IsCompleted();
+		return ConfigFilesTaskEvent.IsValid() && !ConfigFilesTaskEvent.IsCompleted();
 	}
 
 private:
@@ -51,7 +51,7 @@ private:
 	FString GetFilePath(const FString& InConfigId);
 
 	TMap<FString, FSharedBuffer> DownloadedFiles;
-	UE::Tasks::TTask<bool> ConfigFilesTask;
+	UE::Tasks::FTaskEvent ConfigFilesTaskEvent;
 	TWeakPtr<FSubmitToolServiceProvider> ServiceProvider;
 	const FP4LockdownParameters& Parameters;
 	TArray<FAllowListData> AllowListData;
@@ -60,4 +60,4 @@ private:
 	FCriticalSection FileMutex;
 };
 
-Expose_TNameOf(FP4LockdownService);
\ No newline at end of file
+Expose_TNameOf(FP4LockdownService);
