Add support for shelve operation pre-requisite in validators

diff --git a/Engine/Source/Programs/SubmitTool/Private/Logic/Validators/ValidatorDefinition.h b/Engine/Source/Programs/SubmitTool/Private/Logic/Validators/ValidatorDefinition.h
--- a/Engine/Source/Programs/SubmitTool/Private/Logic/Validators/ValidatorDefinition.h
+++ b/Engine/Source/Programs/SubmitTool/Private/Logic/Validators/ValidatorDefinition.h
@@ -248,6 +248,12 @@ struct FValidatorDefinition
 	/// 
 	UPROPERTY()
 	bool bRegisterFailureInCLDescription = false;
+
+	///
+	/// If true, this validator will have submit tool perform a shelve operation before running if it's required.
+	/// 
+	UPROPERTY()
+	bool bRequireShelveOperation = false;
 };
 
 USTRUCT()
diff --git a/Engine/Source/Programs/SubmitTool/Private/Models/ModelInterface.cpp b/Engine/Source/Programs/SubmitTool/Private/Models/ModelInterface.cpp
--- a/Engine/Source/Programs/SubmitTool/Private/Models/ModelInterface.cpp
+++ b/Engine/Source/Programs/SubmitTool/Private/Models/ModelInterface.cpp
@@ -902,7 +902,7 @@ void FModelInterface::OnChangelistReady(bool bIsValid)
 
 				if (!UpdateService->CheckForNewVersion())
 				{
-					ValidationService->QueueByArea(ValidateArea);
+					ValidateByArea(ValidateArea);
 				}
 			}
 
@@ -1219,6 +1219,95 @@ void FModelInterface::OnSwarmCreateCompleted(bool InResult, const FString& InErr
 	}
 }
 
+
+void FModelInterface::ValidateChangelist() const
+{
+	TArray<const TWeakPtr<const FValidatorBase>> Validators;
+	for (const TWeakPtr<const FValidatorBase>& Validator : ValidationService->GetTasks())
+	{
+		const TSharedPtr<const FValidatorBase> SharedValidator = Validator.Pin();
+		if (SharedValidator != nullptr && !SharedValidator->Definition->bIsDisabled && SharedValidator->IsRelevantToCL())
+		{
+			Validators.Add(Validator);
+		}
+	}
+
+	PreValidate(Validators, [this](){ ValidationService->QueueAll(); });
+}
+
+void FModelInterface::ValidateSingle(const FName& ValidatorId, bool bForce) const
+{
+	TArray<const TWeakPtr<const FValidatorBase>> Validators;
+	for (const TWeakPtr<const FValidatorBase>& Validator : ValidationService->GetTasks())
+	{
+		const TSharedPtr<const FValidatorBase> SharedValidator = Validator.Pin();
+		if (SharedValidator != nullptr && SharedValidator->GetValidatorNameId() == ValidatorId && !SharedValidator->Definition->bIsDisabled && SharedValidator->IsRelevantToCL())
+		{
+			Validators.Add(Validator);
+		}
+	}
+
+	PreValidate(Validators, [this, ValidatorId, bForce]() { ValidationService->QueueSingle(ValidatorId, bForce); });
+}
+
+void FModelInterface::ValidateByArea(const ETaskArea InValidateArea) const
+{
+	TArray<const TWeakPtr<const FValidatorBase>> Validators;
+	for (const TWeakPtr<const FValidatorBase>& Validator : ValidationService->GetTasks())
+	{
+		const TSharedPtr<const FValidatorBase> SharedValidator = Validator.Pin();
+		if (SharedValidator != nullptr && (SharedValidator->Definition->TaskArea & InValidateArea) != ETaskArea::None && !SharedValidator->Definition->bIsDisabled && SharedValidator->IsRelevantToCL())
+		{
+			Validators.Add(Validator);
+		}
+	}
+	PreValidate(Validators, [this, InValidateArea]() { ValidationService->QueueByArea(InValidateArea); });
+}
+
+void FModelInterface::PreValidate(const TArray<const TWeakPtr<const FValidatorBase>>& InValidators, TFunction<void(void)> ValidationFunc) const
+{
+	bool bRequiresShelve = false;
+	for (const TWeakPtr<const FValidatorBase>& Validator : InValidators)
+	{
+		const TSharedPtr<const FValidatorBase> SharedValidator = Validator.Pin();
+		if (SharedValidator != nullptr && !SharedValidator->Definition->bIsDisabled && SharedValidator->IsRelevantToCL())
+		{
+			bRequiresShelve |= Validator.Pin()->Definition->bRequireShelveOperation;
+		}
+	}
+
+	if (bRequiresShelve && ChangelistService->GetShelvedFilesDepotPaths() != ChangelistService->GetFilesDepotPaths())
+	{
+		if (ChangelistService->GetShelvedFilesDepotPaths().Num() != 0)
+		{
+			const FText TextTitle = FText::FromString(TEXT("Validation requires a shelf"));
+			const FText TextDescription = FText::FromString(TEXT("Some validations require a shelf, there are already shelved files in this changelist that differ from the local content.\n\nDo you want to delete and recreate your shelf from the latest changes in your local files for use in the preflight?"));
+
+			EDialogFactoryResult DialogResult = FDialogFactory::ShowDialog(TextTitle, TextDescription, TArray<FString>{ TEXT("Re-shelve files"), TEXT("Use existing Shelve")});
+
+			if (DialogResult == EDialogFactoryResult::FirstButton)
+			{
+				ChangelistService->DeleteShelvedFiles(FSourceControlOperationComplete::CreateLambda([ValidationFunc, this](const FSourceControlOperationRef& Operation, ECommandResult::Type Result) {
+					ChangelistService->CreateShelvedFiles(FSourceControlOperationComplete::CreateLambda([ValidationFunc](const FSourceControlOperationRef& Operation, ECommandResult::Type Result) {	ValidationFunc(); }));
+					}));
+			}
+			else
+			{
+				ValidationFunc();
+			}
+		}
+		else
+		{
+			ChangelistService->CreateShelvedFiles(FSourceControlOperationComplete::CreateLambda([ValidationFunc](const FSourceControlOperationRef& Operation, ECommandResult::Type Result) {	ValidationFunc(); }));
+		}
+
+	}
+	else
+	{
+		ValidationFunc();
+	}
+}
+
 void FModelInterface::ToggleValidator(const FName& ValidatorId)
 {
 	ValidationService->ToggleEnabled(ValidatorId);
@@ -1253,4 +1342,3 @@ bool FModelInterface::RemoveTagFromDescription(const FString& InTag, FString& In
 
 	return false;
 }
-
diff --git a/Engine/Source/Programs/SubmitTool/Private/Models/ModelInterface.h b/Engine/Source/Programs/SubmitTool/Private/Models/ModelInterface.h
--- a/Engine/Source/Programs/SubmitTool/Private/Models/ModelInterface.h
+++ b/Engine/Source/Programs/SubmitTool/Private/Models/ModelInterface.h
@@ -126,8 +126,9 @@ public:
 	bool IsBlockingOperationRunning() const											{ return SubmitToolState == ESubmitToolAppState::P4BlockingOperation || SwarmService->IsRequestRunning() || JiraService->IsBlockingRequestRunning() || P4LockdownService->IsBlockingOperationRunning(); }
 	void CancelP4Operations(FName OperationName = FName()) const					{ ChangelistService->CancelP4Operations(OperationName); SwarmService->CancelOperations(); }
 
-	void ValidateChangelist() const													{ ValidationService->QueueAll(); }
-	void ValidateSingle(const FName& ValidatorId, bool bForce = true)				{ ValidationService->QueueSingle(ValidatorId, bForce); }
+	void ValidateChangelist() const;
+	void ValidateSingle(const FName& ValidatorId, bool bForce = true) const;
+	void ValidateByArea(const ETaskArea InValidateArea) const;
 	void ToggleValidator(const FName& ValidatorId);
 	void ValidateCLDescription() const												{ ValidationService->StopTasksByArea(ETaskArea::Changelist); ValidationService->QueueByArea(ETaskArea::Changelist); }
 	bool IsCLValid() const															{ return ValidationService->GetIsRunSuccessful(!IsIntegrationRequired()); }
@@ -264,6 +265,7 @@ private:
 
 		ValidationService->InvalidateForChanges(InChangeType);
 	}
+	void PreValidate(const TArray<const TWeakPtr<const FValidatorBase>>& InValidators, TFunction<void(void)> ValidationFunc) const;
 	void RefreshStateBasedOnFiles();
 	void OnChangelistReady(bool bIsValid);
 	void RevertUnchangedAndSubmit();
