SubmitTool: SwarmService: Check that the FSwarmService::FetchReview http response pointer is valid before using it.

#rnx

diff --git a/Engine/Source/Programs/SubmitTool/Private/Logic/SwarmService.cpp b/Engine/Source/Programs/SubmitTool/Private/Logic/SwarmService.cpp
--- a/Engine/Source/Programs/SubmitTool/Private/Logic/SwarmService.cpp
+++ b/Engine/Source/Programs/SubmitTool/Private/Logic/SwarmService.cpp
@@ -53,7 +53,7 @@ void FSwarmService::FetchReview(const OnGetReviewComplete& OnComplete)
 
 	HttpRequest->OnProcessRequestComplete().BindLambda([OnComplete, this](FHttpRequestPtr Request, FHttpResponsePtr Response, bool bConnectedSuccessfully)
 		{
-			if (!bConnectedSuccessfully)
+			if (!bConnectedSuccessfully || !Response.IsValid())
 			{
 				OnComplete.ExecuteIfBound(Review, TEXT("Connection Failed"));
 				return;
