Fixing issue where warning were incorrectly being printed as errors when using anchoring

diff --git a/Engine/Source/Programs/SubmitTool/Private/Logic/Validators/ValidatorBase.cpp b/Engine/Source/Programs/SubmitTool/Private/Logic/Validators/ValidatorBase.cpp
--- a/Engine/Source/Programs/SubmitTool/Private/Logic/Validators/ValidatorBase.cpp
+++ b/Engine/Source/Programs/SubmitTool/Private/Logic/Validators/ValidatorBase.cpp
@@ -65,6 +65,7 @@ void FValidatorBase::StartValidation()
 	Start = FDateTime::UtcNow();
 	State = EValidationStates::Running;
 	ErrorListCache.Empty();
+	WarningListCache.Empty();
 
 	const TSharedPtr<FChangelistService>& ChangelistService = ServiceProvider.Pin()->GetService<FChangelistService>();
 	const TSharedPtr<FTagService>& TagService = ServiceProvider.Pin()->GetService<FTagService>();
@@ -433,8 +434,15 @@ void FValidatorBase::PrintErrorSummary() const
 				UE_LOG(LogValidatorsResult, Warning, TEXT("%s"), *ErrorStr);
 			}
 		}
+
+		for (const FString& WarningStr : WarningListCache)
+		{
+			UE_LOG(LogValidators, Warning, TEXT("%s"), *WarningStr);
+			UE_LOG(LogValidatorsResult, Warning, TEXT("%s"), *WarningStr);
+		}
 	}
 }
+
 const FString FValidatorBase::GetValidationConfigId() const
 {
 	TStringBuilder<512> StringBuilder;
diff --git a/Engine/Source/Programs/SubmitTool/Private/Logic/Validators/ValidatorBase.h b/Engine/Source/Programs/SubmitTool/Private/Logic/Validators/ValidatorBase.h
--- a/Engine/Source/Programs/SubmitTool/Private/Logic/Validators/ValidatorBase.h
+++ b/Engine/Source/Programs/SubmitTool/Private/Logic/Validators/ValidatorBase.h
@@ -220,6 +220,7 @@ protected:
 	TArray<FSourceControlStateRef> FilteredFiles;
 	TArray<FString> ActivationErrors;
 	mutable TArray<FString> ErrorListCache;
+	mutable TArray<FString> WarningListCache;
 	bool bIsValidSetup = false;
 	bool bForceRun = false;
 	const FSubmitToolParameters& SubmitToolParameters;
diff --git a/Engine/Source/Programs/SubmitTool/Private/Logic/Validators/ValidatorRunExecutable.cpp b/Engine/Source/Programs/SubmitTool/Private/Logic/Validators/ValidatorRunExecutable.cpp
--- a/Engine/Source/Programs/SubmitTool/Private/Logic/Validators/ValidatorRunExecutable.cpp
+++ b/Engine/Source/Programs/SubmitTool/Private/Logic/Validators/ValidatorRunExecutable.cpp
@@ -221,17 +221,18 @@ bool FValidatorRunExecutable::QueueProcess(const FString& InId, const FString& L
 void FValidatorRunExecutable::OnProcessOutputLine(const FString& ProcessId, const FString& Line, const EProcessOutputType& OutputType)
 {
 	const FString& FormattedMessage = FString::Printf(TEXT("[%s (%s)]: %s"), *GetValidatorName(), *ProcessId, *Line);
-	if(OutputType == EProcessOutputType::ProcessError || IsLineAnError(Line, ProcessesData[ProcessId].bIgnoringOutputError, ProcessesData[ProcessId].bIgnoringWarningsAsErrors))
+	bool bIsWarning = false;
+	if(OutputType == EProcessOutputType::ProcessError || IsLineAnError(Line, bIsWarning, ProcessesData[ProcessId].bIgnoringOutputError, ProcessesData[ProcessId].bIgnoringWarningsAsErrors))
 	{
 		LogFailure(FormattedMessage);
 		ProcessesData[ProcessId].ErrorList.Add(FormattedMessage);
 	}
 	else
 	{
-		if (GetTypedDefinition<FValidatorRunExecutableDefinition>()->bIncludeWarningsInErrorSummary)
+		if (bIsWarning && GetTypedDefinition<FValidatorRunExecutableDefinition>()->bIncludeWarningsInErrorSummary)
 		{
 			FScopeLock Lock(&ErrorLogsMutex);
-			ProcessesData[ProcessId].ErrorList.Add(FormattedMessage);
+			ProcessesData[ProcessId].WarningList.Add(FormattedMessage);
 			ErrorListCache.Add(FormattedMessage);
 		}
 		
@@ -239,7 +240,7 @@ void FValidatorRunExecutable::OnProcessOutputLine(const FString& ProcessId, cons
 	}
 }
 
-bool FValidatorRunExecutable::IsLineAnError(const FString& InLine, bool& InOutbIgnoringOutputErrors, bool& InOutbIgnoringWarningsAsErrors)
+bool FValidatorRunExecutable::IsLineAnError(const FString& InLine, bool& OutbIsWarning, bool& InOutbIgnoringOutputErrors, bool& InOutbIgnoringWarningsAsErrors)
 {
 	const FValidatorRunExecutableDefinition* TypedDefinition = GetTypedDefinition<FValidatorRunExecutableDefinition>();
 
@@ -297,10 +298,12 @@ bool FValidatorRunExecutable::IsLineAnError(const FString& InLine, bool& InOutbI
 		}
 	}
 
-	if(Definition->bTreatWarningsAsErrors && !InOutbIgnoringWarningsAsErrors)
+	FRegexMatcher Regex = FRegexMatcher(*TypedDefinition->BuiltRegexWarning, InLine);
+	if (Regex.FindNext())
 	{
-		FRegexMatcher Regex = FRegexMatcher(*TypedDefinition->BuiltRegexWarning, InLine);
-		if (Regex.FindNext())
+		OutbIsWarning = true;
+
+		if(Definition->bTreatWarningsAsErrors && !InOutbIgnoringWarningsAsErrors)
 		{
 			return true;
 		}
diff --git a/Engine/Source/Programs/SubmitTool/Private/Logic/Validators/ValidatorRunExecutable.h b/Engine/Source/Programs/SubmitTool/Private/Logic/Validators/ValidatorRunExecutable.h
--- a/Engine/Source/Programs/SubmitTool/Private/Logic/Validators/ValidatorRunExecutable.h
+++ b/Engine/Source/Programs/SubmitTool/Private/Logic/Validators/ValidatorRunExecutable.h
@@ -11,6 +11,7 @@ struct FProcessTrackingData
 	bool bIgnoringOutputError;
 	bool bIgnoringWarningsAsErrors;
 	TArray<FString> ErrorList;
+	TArray<FString> WarningList;
 };
 
 /**
@@ -33,7 +34,7 @@ public:
 protected:
 	virtual bool QueueProcess(const FString& ProcessId, const FString& LocalPath, const FString& Args);
 	virtual void OnProcessOutputLine(const FString& ProcessId, const FString& Line, const EProcessOutputType& OutputType);
-	virtual bool IsLineAnError(const FString& InLine, bool& InOutbIgnoringOutputErrors, bool& InOutbIgnoringWarningsAsErrors);
+	virtual bool IsLineAnError(const FString& InLine, bool& OutbIsWarning, bool& InOutbIgnoringOutputErrors, bool& InOutbIgnoringWarningsAsErrors);
 	virtual void OnProcessComplete(const FString& ProcessId, int32 ReturnCode);
 	
 	virtual void Tick(float DeltaTime) override;
