Print warnings in summary even when the validator has succeeded

diff --git a/Engine/Source/Programs/SubmitTool/Private/Logic/TasksService.cpp b/Engine/Source/Programs/SubmitTool/Private/Logic/TasksService.cpp
--- a/Engine/Source/Programs/SubmitTool/Private/Logic/TasksService.cpp
+++ b/Engine/Source/Programs/SubmitTool/Private/Logic/TasksService.cpp
@@ -184,8 +184,8 @@ void FTasksService::PrintErrorSummary()
 		{
 			if (!bHeaderPrinted)
 			{
-				UE_LOG(LogValidators, Error, TEXT("========================[Errors Summary #%d]========================"), Execution);
-				UE_LOG(LogValidatorsResult, Error, TEXT("========================[Errors Summary #%d]========================"), Execution);
+				UE_LOG(LogValidators, Error, TEXT("========================[Summary #%d]========================"), Execution);
+				UE_LOG(LogValidatorsResult, Error, TEXT("========================[Summary #%d]========================"), Execution);
 				bHeaderPrinted = true;
 			}
 			TaskPair.Value->PrintErrorSummary();
diff --git a/Engine/Source/Programs/SubmitTool/Private/Logic/Validators/ValidatorBase.cpp b/Engine/Source/Programs/SubmitTool/Private/Logic/Validators/ValidatorBase.cpp
--- a/Engine/Source/Programs/SubmitTool/Private/Logic/Validators/ValidatorBase.cpp
+++ b/Engine/Source/Programs/SubmitTool/Private/Logic/Validators/ValidatorBase.cpp
@@ -402,9 +402,9 @@ void FValidatorBase::SetSelectedOption(const FString& InOptionName, const FStrin
 
 bool FValidatorBase::CanPrintErrors() const
 {
-	if (State == EValidationStates::Failed || State == EValidationStates::Timeout)
+	if (State == EValidationStates::Failed || State == EValidationStates::Timeout || State == EValidationStates::Valid)
 	{
-		return !ErrorListCache.IsEmpty();
+		return !ErrorListCache.IsEmpty() || !WarningListCache.IsEmpty();
 	}
 	else
 	{
diff --git a/Engine/Source/Programs/SubmitTool/Private/Logic/Validators/ValidatorRunExecutable.cpp b/Engine/Source/Programs/SubmitTool/Private/Logic/Validators/ValidatorRunExecutable.cpp
--- a/Engine/Source/Programs/SubmitTool/Private/Logic/Validators/ValidatorRunExecutable.cpp
+++ b/Engine/Source/Programs/SubmitTool/Private/Logic/Validators/ValidatorRunExecutable.cpp
@@ -233,10 +233,15 @@ void FValidatorRunExecutable::OnProcessOutputLine(const FString& ProcessId, cons
 		{
 			FScopeLock Lock(&ErrorLogsMutex);
 			ProcessesData[ProcessId].WarningList.Add(FormattedMessage);
-			ErrorListCache.Add(FormattedMessage);
+			WarningListCache.Add(FormattedMessage);
+			UE_LOG(LogValidatorsResult, Warning, TEXT("%s"), *FormattedMessage);
+			UE_LOG(LogValidators, Warning, TEXT("%s"), *FormattedMessage);
+
 		}
-		
-		UE_LOG(LogValidators, Log, TEXT("%s"), *FormattedMessage);
+		else
+		{
+			UE_LOG(LogValidators, Log, TEXT("%s"), *FormattedMessage);
+		}		
 	}
 }
 
