Adding support for using changelist addendum in failures

diff --git a/Engine/Source/Programs/SubmitTool/Private/Logic/TasksService.cpp b/Engine/Source/Programs/SubmitTool/Private/Logic/TasksService.cpp
--- a/Engine/Source/Programs/SubmitTool/Private/Logic/TasksService.cpp
+++ b/Engine/Source/Programs/SubmitTool/Private/Logic/TasksService.cpp
@@ -469,9 +469,13 @@ const TArray<FString> FTasksService::GetAddendums() const
 	TArray<FString> Addendums;
 	for(const TPair<FName, TSharedPtr<FValidatorBase>>& TaskPair : Tasks)
 	{
-		if(!TaskPair.Value->Definition->ChangelistDescriptionAddendum.IsEmpty() && TaskPair.Value->GetState() == EValidationStates::Valid)
+		if(!TaskPair.Value->Definition->ChangelistDescriptionAddendum.IsEmpty())
 		{
-			Addendums.Emplace(TaskPair.Value->Definition->ChangelistDescriptionAddendum);
+			if (TaskPair.Value->GetState() == EValidationStates::Valid ||
+				(TaskPair.Value->Definition->bAppendAddendumOnAnyResult && TaskPair.Value->GetState() == EValidationStates::Failed))
+			{
+				Addendums.Emplace(TaskPair.Value->Definition->ChangelistDescriptionAddendum);
+			}
 		}
 	}
 
