missing file with bAppendAddendumOnAnyResult property

diff --git a/Engine/Source/Programs/SubmitTool/Private/Logic/Validators/ValidatorDefinition.h b/Engine/Source/Programs/SubmitTool/Private/Logic/Validators/ValidatorDefinition.h
--- a/Engine/Source/Programs/SubmitTool/Private/Logic/Validators/ValidatorDefinition.h
+++ b/Engine/Source/Programs/SubmitTool/Private/Logic/Validators/ValidatorDefinition.h
@@ -146,6 +146,12 @@ struct FValidatorDefinition
 	UPROPERTY()
 	FString ChangelistDescriptionAddendum;
 
+	///
+	///	If true, addendum will be appended as long as the validation has been run, even on failure
+	/// 
+	UPROPERTY()
+	bool bAppendAddendumOnAnyResult = false;
+
 	///
 	///	Skip this validator when the addendum is already present in the CL description
 	/// 
