Adding a better mechanism for knowing when to reshelve files when they need refresh for validators that require it

#code-analyzed

diff --git a/Engine/Source/Programs/SubmitTool/Private/Logic/ChangelistService.cpp b/Engine/Source/Programs/SubmitTool/Private/Logic/ChangelistService.cpp
--- a/Engine/Source/Programs/SubmitTool/Private/Logic/ChangelistService.cpp
+++ b/Engine/Source/Programs/SubmitTool/Private/Logic/ChangelistService.cpp
@@ -487,6 +487,7 @@ void FChangelistService::CreateShelvedFiles(const FSourceControlOperationComplet
 			{
 				ShelvedFilesInCL = ChangelistState->GetShelvedFilesStates();
 				ShelvedFilesDepotPaths.Reset(ShelvedFilesInCL.Num());
+				LastShelfTime = FDateTime::UtcNow();
 				UE_LOG(LogSubmitToolP4, Log, TEXT("CL %s create shelved operation succeeded"), *GetCLID());
 			}
 			else if(Result == ECommandResult::Cancelled)
diff --git a/Engine/Source/Programs/SubmitTool/Private/Logic/ChangelistService.h b/Engine/Source/Programs/SubmitTool/Private/Logic/ChangelistService.h
--- a/Engine/Source/Programs/SubmitTool/Private/Logic/ChangelistService.h
+++ b/Engine/Source/Programs/SubmitTool/Private/Logic/ChangelistService.h
@@ -67,6 +67,11 @@ public:
 		return true;
 	}
 
+	const FDateTime& GetLastShelfDateTime() const
+	{
+		return LastShelfTime;
+	}
+
 	const TArray<FSourceControlStateRef>& GetFilesInCL() const
 	{
 		return FilesInCL;
@@ -130,6 +135,7 @@ private:
 	FText OriginalDescription;
 	FString CLDescription;
 	ISourceControlProvider* CachedSCCProvider;
+	FDateTime LastShelfTime = FDateTime::MinValue();
 	const TSharedPtr<ISTSourceControlService> SourceControlService;
 
 	void FindInitialChangelistsAsync();
diff --git a/Engine/Source/Programs/SubmitTool/Private/Models/ModelInterface.cpp b/Engine/Source/Programs/SubmitTool/Private/Models/ModelInterface.cpp
--- a/Engine/Source/Programs/SubmitTool/Private/Models/ModelInterface.cpp
+++ b/Engine/Source/Programs/SubmitTool/Private/Models/ModelInterface.cpp
@@ -14,6 +14,7 @@
 #include "Configuration/Configuration.h"
 #include "SourceControlOperations.h"
 #include "HAL/PlatformApplicationMisc.h"
+#include "HAL/FileManagerGeneric.h"
 #include "Async/Async.h"
 #include "Misc/StringOutputDevice.h"
 #include "Telemetry/TelemetryService.h"
@@ -1276,36 +1277,46 @@ void FModelInterface::PreValidate(const TArray<const TWeakPtr<const FValidatorBa
 		}
 	}
 
-	if (bRequiresShelve && ChangelistService->GetShelvedFilesDepotPaths() != ChangelistService->GetFilesDepotPaths())
+	if (bRequiresShelve)
 	{
 		if (ChangelistService->GetShelvedFilesDepotPaths().Num() != 0)
 		{
-			const FText TextTitle = FText::FromString(TEXT("Validation requires a shelf"));
-			const FText TextDescription = FText::FromString(TEXT("Some validations require a shelf, there are already shelved files in this changelist that differ from the local content.\n\nDo you want to delete and recreate your shelf from the latest changes in your local files for use in the preflight?"));
-
-			EDialogFactoryResult DialogResult = FDialogFactory::ShowDialog(TextTitle, TextDescription, TArray<FString>{ TEXT("Re-shelve files"), TEXT("Use existing Shelve")});
-
-			if (DialogResult == EDialogFactoryResult::FirstButton)
+			FDateTime LastModifiedFile;
+			FFileManagerGeneric FileManager;
+			for (const FSourceControlStateRef& File : ChangelistService->GetFilesInCL())
 			{
-				ChangelistService->DeleteShelvedFiles(FSourceControlOperationComplete::CreateLambda([ValidationFunc, this](const FSourceControlOperationRef& Operation, ECommandResult::Type Result) {
-					ChangelistService->CreateShelvedFiles(FSourceControlOperationComplete::CreateLambda([ValidationFunc](const FSourceControlOperationRef& Operation, ECommandResult::Type Result) {	ValidationFunc(); }));
-					}));
+				FString Filename = File->GetFilename();
+				FFileStatData FileModifiedDate = FileManager.GetStatData(*Filename);
+				if (FileModifiedDate.bIsValid && FileModifiedDate.ModificationTime > LastModifiedFile)
+				{
+					LastModifiedFile = FileModifiedDate.ModificationTime;
+				}
 			}
-			else
+
+			if (ChangelistService->GetLastShelfDateTime() < LastModifiedFile)
 			{
-				ValidationFunc();
-			}
+				const FText TextTitle = FText::FromString(TEXT("Validation requires a shelf"));
+				const FText TextDescription = FText::FromString(TEXT("Some validations require a shelf, there are already shelved files in this changelist that differ from the local content.\n\nDo you want to delete and recreate your shelf from the latest changes in your local files for use in the preflight?"));
+
+				EDialogFactoryResult DialogResult = FDialogFactory::ShowDialog(TextTitle, TextDescription, TArray<FString>{ TEXT("Re-shelve files"), TEXT("Use existing Shelve")});
+
+				if (DialogResult == EDialogFactoryResult::FirstButton)
+				{
+					ChangelistService->DeleteShelvedFiles(FSourceControlOperationComplete::CreateLambda([ValidationFunc, this](const FSourceControlOperationRef& Operation, ECommandResult::Type Result) {
+						ChangelistService->CreateShelvedFiles(FSourceControlOperationComplete::CreateLambda([ValidationFunc](const FSourceControlOperationRef& Operation, ECommandResult::Type Result) { ValidationFunc(); }));
+						}));
+					return;
+				}
+			}			
 		}
 		else
 		{
-			ChangelistService->CreateShelvedFiles(FSourceControlOperationComplete::CreateLambda([ValidationFunc](const FSourceControlOperationRef& Operation, ECommandResult::Type Result) {	ValidationFunc(); }));
+			ChangelistService->CreateShelvedFiles(FSourceControlOperationComplete::CreateLambda([ValidationFunc](const FSourceControlOperationRef& Operation, ECommandResult::Type Result) { ValidationFunc(); }));
+			return;
 		}
-
-	}
-	else
-	{
-		ValidationFunc();
 	}
+
+	ValidationFunc();
 }
 
 void FModelInterface::ToggleValidator(const FName& ValidatorId)
