Adding support for skipping validators when a specific preflight template has succeeded
Adding support for checking out of memory potential before running validators, Validators can optionally specify their expected memory consumption. current implementation doesn't take into account other validators and assumes it's running in isolation for now.

#code-analyzed

diff --git a/Engine/Source/Programs/SubmitTool/Private/Logic/PreflightService.cpp b/Engine/Source/Programs/SubmitTool/Private/Logic/PreflightService.cpp
--- a/Engine/Source/Programs/SubmitTool/Private/Logic/PreflightService.cpp
+++ b/Engine/Source/Programs/SubmitTool/Private/Logic/PreflightService.cpp
@@ -128,8 +128,9 @@ bool FPreflightService::Tick(float DeltaTime)
 	return true;
 }
 
-TMap<FString, FStringFormatArg> FPreflightService::GetFormatParameters() const
+TMap<FString, FStringFormatArg> FPreflightService::GetFormatParameters(const FString& InTemplate)
 {
+	RefreshStreamName();
 	TMap<FString, FStringFormatArg> FormatMap =
 	{
 		{ TEXT("URL"), Definition.HordeServerAddress },
@@ -137,29 +138,34 @@ TMap<FString, FStringFormatArg> FPreflightService::GetFormatParameters() const
 	};
 
 	FormatMap.Add(TEXT("Stream"), StreamName);
-	FormatMap.Add(TEXT("Template"), FString());
+	FormatMap.Add(TEXT("Template"), InTemplate);
 	FormatMap.Add(TEXT("AdditionalTasks"), FString());
 
-	FPreflightTemplateDefinition Template;
-	if(SelectPreflightTemplate(Template))
-	{
-		FormatMap[TEXT("Template")] = Template.Template;
-		FormatMap[TEXT("AdditionalTasks")] = GetAdditionalTasksString(Template);
-	}
-	else
+
+	if (InTemplate.IsEmpty())
 	{
-		FormatMap[TEXT("Template")] = Definition.DefaultPreflightTemplate;
+		FPreflightTemplateDefinition Template;
+		if(SelectPreflightTemplate(Template))
+		{
+			FormatMap[TEXT("Template")] = Template.Template;
+			FormatMap[TEXT("AdditionalTasks")] = GetAdditionalTasksString(Template);
+		}
+		else
+		{
+			FormatMap[TEXT("Template")] = Definition.DefaultPreflightTemplate;
+		}
 	}
 
 	return FormatMap;
 }
 
-void FPreflightService::RequestPreflight()
+void FPreflightService::RequestPreflight(const FString InTemplate)
 {
 	if(State == EPreflightServiceState::Idle)
 	{
 		if(!ServiceProvider.Pin()->GetService<ISTSourceControlService>()->GetClientStreams().IsEmpty())
 		{
+			TemplateId = InTemplate;
 			UE_LOG(LogSubmitTool, Log, TEXT("Preflight: Requesting..."));
 			State = EPreflightServiceState::RequestDeleteShelve;
 		}
@@ -209,9 +215,12 @@ void FPreflightService::Requeue()
 
 void FPreflightService::FetchPreflightInfo(bool bRequeue, const FString& InOAuthToken)
 {
+	TRACE_CPUPROFILER_EVENT_SCOPE(FPreflightService::FetchPreflightInfo);
+
 	TSharedPtr<FCredentialsService> CredentialsService = ServiceProvider.Pin()->GetService<FCredentialsService>();
 	if(Definition.HordeServerAddress.IsEmpty() || !CredentialsService->IsOIDCTokenEnabled() || FModelInterface::GetState() == ESubmitToolAppState::Finished)
 	{
+		bHordeInformationReady = true;
 		return;
 	}
 
@@ -249,12 +258,12 @@ void FPreflightService::FetchPreflightInfo(bool bRequeue, const FString& InOAuth
 		// ensure the token is the most up to date
 		LinkedPFRequest->SetHeader(TEXT("Authorization"), FString::Printf(TEXT("bearer %s"), *OIDCToken));
 
-		if(!LinkedPFRequest->OnProcessRequestComplete().IsBound())
+		if(LinkedPFRequest->OnProcessRequestComplete().IsBound())
 		{
 			LinkedPFRequest->OnProcessRequestComplete().Unbind();
 		}
 
-		LinkedPFRequest->OnProcessRequestComplete().BindLambda([this, bRequeue, &OIDCToken](FHttpRequestPtr Request, FHttpResponsePtr HttpResponse, bool bConnectedSuccessfully)
+		LinkedPFRequest->OnProcessRequestComplete().BindLambda([this, bRequeue, OIDCToken](FHttpRequestPtr Request, FHttpResponsePtr HttpResponse, bool bConnectedSuccessfully)
 			{
 				if(!bConnectedSuccessfully)
 				{
@@ -268,6 +277,7 @@ void FPreflightService::FetchPreflightInfo(bool bRequeue, const FString& InOAuth
 						UE_LOG(LogSubmitTool, Warning, TEXT("Unable to connect to horde. Connection error, no response."));
 					}
 
+					bHordeInformationReady = true;
 					OnHordeConnectionFailed.Broadcast();
 					return;
 				}
@@ -343,7 +353,12 @@ void FPreflightService::FetchPreflightInfo(bool bRequeue, const FString& InOAuth
 										}
 										else
 										{
-											EDialogFactoryResult Result = ShowUpdatePreflightTagDialog();
+											EDialogFactoryResult Result = EDialogFactoryResult::FirstButton;
+											if (!CurrentTagValue.IsEmpty())
+											{
+												Result = ShowUpdatePreflightTagDialog();
+											}
+
 											if(Result == EDialogFactoryResult::FirstButton)
 											{
 												// Set the latest one as the tag value
@@ -375,12 +390,13 @@ void FPreflightService::FetchPreflightInfo(bool bRequeue, const FString& InOAuth
 				}
 				else
 				{
-					UE_LOG(LogSubmitTool, Warning, TEXT("Unable to fetch preflights. Failed with code %d"), HttpResponse->GetResponseCode());
-					UE_LOG(LogSubmitToolDebug, Warning, TEXT("Unable to fetch preflights. Failed with code %d\nResponse: %s"), HttpResponse->GetResponseCode(), *HttpResponse->GetContentAsString());
+					UE_LOG(LogSubmitTool, Warning, TEXT("Unable to fetch preflights. Failed unknown response"));
+					UE_LOG(LogSubmitToolDebug, Warning, TEXT("Unable to fetch preflights. Failed unknown response"));
 				}
 
 				if(ActiveUnlinkedRequests == 0)
 				{
+					bHordeInformationReady = true;
 					if(OnPreflightDataUpdated.IsBound() && HordePreflights.IsValid())
 					{
 						OnPreflightDataUpdated.Broadcast(HordePreflights, UnlinkedHordePreflights);
@@ -446,13 +462,13 @@ void FPreflightService::FetchUnlinkedPreflight(const FString& InPreflightId, boo
 
 	UnlinkedPFRequest->SetHeader(TEXT("Authorization"), FString::Printf(TEXT("bearer %s"), *InOAuthToken));
 
-	if(!UnlinkedPFRequest->OnProcessRequestComplete().IsBound())
+	if(UnlinkedPFRequest->OnProcessRequestComplete().IsBound())
 	{
 		UnlinkedPFRequest->OnProcessRequestComplete().Unbind();
 	}
 
 	ActiveUnlinkedRequests++;
-	UnlinkedPFRequest->OnProcessRequestComplete().BindLambda([this, bRequeue, &InOAuthToken, InPreflightId](FHttpRequestPtr Request, FHttpResponsePtr HttpResponse, bool bConnectedSuccessfully)
+	UnlinkedPFRequest->OnProcessRequestComplete().BindLambda([this, bRequeue, InOAuthToken, InPreflightId](FHttpRequestPtr Request, FHttpResponsePtr HttpResponse, bool bConnectedSuccessfully)
 		{
 			ActiveUnlinkedRequests--;
 			if(!bConnectedSuccessfully)
@@ -503,13 +519,13 @@ void FPreflightService::FetchUnlinkedPreflight(const FString& InPreflightId, boo
 			}
 			else
 			{
-				UE_LOG(LogSubmitTool, Warning, TEXT("Unable to fetch preflights. Failed with code %d"), HttpResponse->GetResponseCode());
-				UE_LOG(LogSubmitToolDebug, Warning, TEXT("Unable to fetch preflights. Failed with code %d\nResponse: %s"), HttpResponse->GetResponseCode(), *HttpResponse->GetContentAsString());
+				UE_LOG(LogSubmitTool, Warning, TEXT("Unable to fetch preflights. Failed unknown response"));
+				UE_LOG(LogSubmitToolDebug, Warning, TEXT("Unable to fetch preflights. Failed unknown response"));
 			}
 
 			if(ActiveUnlinkedRequests == 0)
 			{
-				if(OnPreflightDataUpdated.IsBound() && HordePreflights.IsValid())
+				if(OnPreflightDataUpdated.IsBound())
 				{
 					OnPreflightDataUpdated.Broadcast(HordePreflights, UnlinkedHordePreflights);
 				}
@@ -527,9 +543,6 @@ void FPreflightService::FetchUnlinkedPreflight(const FString& InPreflightId, boo
 
 void FPreflightService::StartPreflight()
 {
-	TSharedPtr<ISTSourceControlService> SCCService = ServiceProvider.Pin()->GetService<ISTSourceControlService>();
-	StreamName = SCCService->GetRootStreamName();
-
 	const TArray<FString>& ShelvedPaths = ServiceProvider.Pin()->GetService<FChangelistService>()->GetShelvedFilesDepotPaths(true);
 
 	if (ShelvedPaths.IsEmpty())
@@ -539,41 +552,8 @@ void FPreflightService::StartPreflight()
 		return;
 	}
 
-	FString CommonPath = ShelvedPaths[0];
-	const FString& LastPath = ShelvedPaths.Last();
-
-	for(size_t i = 0; i < FMath::Min(CommonPath.Len(), LastPath.Len()); ++i)
-	{
-		if(CommonPath[i] != LastPath[i])
-		{
-			CommonPath = CommonPath.Left(i);
-			break;
-		}
-	}
-
-	if(!CommonPath.Equals(TEXT("//")))
-	{
-		int32 NextSlash = CommonPath.Find(TEXT("/"), ESearchCase::IgnoreCase, ESearchDir::FromStart, 2);
-		const FString Depot = CommonPath.Mid(2, NextSlash - 2);
-		size_t StreamDepth = SCCService->GetDepotStreamLength(Depot);
-
-		for(size_t i = 0; i < StreamDepth; ++i)
-		{
-			NextSlash = CommonPath.Find(TEXT("/"), ESearchCase::IgnoreCase, ESearchDir::FromStart, NextSlash+1);
-		}
-
-		if(NextSlash != -1)
-		{
-			const FSCCStream* FoundStream = SCCService->GetSCCStream(CommonPath.Left(NextSlash));
-
-			if(FoundStream != nullptr)
-			{
-				StreamName = FoundStream->Name;
-			}
-		}
-	}
-
-	FString StartPreflightUrl = FString::Format(*Definition.StartPreflightURLFormat, GetFormatParameters());
+	FString StartPreflightUrl = FString::Format(*Definition.StartPreflightURLFormat, GetFormatParameters(TemplateId));
+	TemplateId = FString();
 
 	// If for some reason, our preflight settings are missing, this will be empty, let's not popup a browser with nothing in it
 	if(!StartPreflightUrl.IsEmpty())
@@ -588,7 +568,8 @@ void FPreflightService::StartPreflight()
 		FPlatformProcess::LaunchURL(*StartPreflightUrl, nullptr, nullptr);
 		State = EPreflightServiceState::Idle;
 
-		// Do a Fetch in 10 and 30 s to try and capture the triggered preflight
+		// Do a Fetch in 5, 10 and 30 s to try and capture the triggered preflight
+		QueueFetch(false, 5.f);
 		QueueFetch(false, 10.f);
 		QueueFetch(false, 30.f);
 	}
@@ -641,6 +622,55 @@ void FPreflightService::OnCreateShelveOperationComplete(const FSourceControlOper
 	}
 }
 
+TArray<const FPreflightData*> FPreflightService::GetTaggedPreflights() const
+{
+	TArray<const FPreflightData*> FoundPreflights;
+	if (PreflightTag == nullptr)
+	{
+		return FoundPreflights;
+	}
+
+	const TArray<FString>& PreflightValues = PreflightTag->GetValues();
+
+	for (FString PreflightId : PreflightValues)
+	{
+		if (PreflightId.Contains(TEXT("/")))
+		{
+			int32 SlashIdx;
+			PreflightId.FindLastChar(TCHAR('/'), SlashIdx);
+			PreflightId = PreflightId.RightChop(SlashIdx + 1);
+		}
+		PreflightId.TrimStartAndEndInline();
+
+		FRegexPattern Pattern = FRegexPattern(TEXT("(?:[0-9]|[a-f]){24}"), ERegexPatternFlags::CaseInsensitive);
+		FRegexMatcher regex = FRegexMatcher(Pattern, PreflightId);
+		if (regex.FindNext())
+		{
+			const FPreflightData* FoundData = nullptr;
+			if (HordePreflights.IsValid())
+			{
+				FoundData = HordePreflights->PreflightList.FindByPredicate([&PreflightId](const FPreflightData& InData) { return InData.ID == PreflightId; });
+				// Look into the linked preflights and if we don't find it check the unlinked preflights
+				if (FoundData != nullptr)
+				{
+					FoundPreflights.Add(FoundData);
+				}
+			}
+
+			if(FoundData == nullptr)
+			{
+				FoundData = UnlinkedHordePreflights.Find(PreflightId);
+				if (FoundData)
+				{
+					FoundPreflights.Add(FoundData);
+				}
+			}
+		}
+	}
+
+	return FoundPreflights;
+}
+
 bool FPreflightService::SelectPreflightTemplate(FPreflightTemplateDefinition& OutTemplate) const
 {
 	TSharedPtr<FChangelistService> ChangelistService = ServiceProvider.Pin()->GetService<FChangelistService>();
@@ -694,6 +724,51 @@ FString FPreflightService::GetAdditionalTasksString(const FPreflightTemplateDefi
 	return AdditionalTaskStrBuilder.ToString();
 }
 
+void FPreflightService::RefreshStreamName()
+{
+	TSharedPtr<ISTSourceControlService> SCCService = ServiceProvider.Pin()->GetService<ISTSourceControlService>();
+	StreamName = SCCService->GetRootStreamName();
+
+	const TArray<FString>& ShelvedPaths = ServiceProvider.Pin()->GetService<FChangelistService>()->GetShelvedFilesDepotPaths(true);
+
+	if (ShelvedPaths.Num() != 0)
+	{
+		FString CommonPath = ShelvedPaths[0];
+		const FString& LastPath = ShelvedPaths.Last();
+
+		for (size_t i = 0; i < FMath::Min(CommonPath.Len(), LastPath.Len()); ++i)
+		{
+			if (CommonPath[i] != LastPath[i])
+			{
+				CommonPath = CommonPath.Left(i);
+				break;
+			}
+		}
+
+		if (!CommonPath.Equals(TEXT("//")))
+		{
+			int32 NextSlash = CommonPath.Find(TEXT("/"), ESearchCase::IgnoreCase, ESearchDir::FromStart, 2);
+			const FString Depot = CommonPath.Mid(2, NextSlash - 2);
+			size_t StreamDepth = SCCService->GetDepotStreamLength(Depot);
+
+			for (size_t i = 0; i < StreamDepth; ++i)
+			{
+				NextSlash = CommonPath.Find(TEXT("/"), ESearchCase::IgnoreCase, ESearchDir::FromStart, NextSlash + 1);
+			}
+
+			if (NextSlash != -1)
+			{
+				const FSCCStream* FoundStream = SCCService->GetSCCStream(CommonPath.Left(NextSlash));
+
+				if (FoundStream != nullptr)
+				{
+					StreamName = FoundStream->Name;
+				}
+			}
+		}
+	}
+}
+
 EDialogFactoryResult FPreflightService::ShowRecreateShelveDialog() const
 {
 	if(ModelInterface->GetMainTab().IsValid() && ModelInterface->GetMainTab().Pin()->GetParentWindow().IsValid())
diff --git a/Engine/Source/Programs/SubmitTool/Private/Logic/PreflightService.h b/Engine/Source/Programs/SubmitTool/Private/Logic/PreflightService.h
--- a/Engine/Source/Programs/SubmitTool/Private/Logic/PreflightService.h
+++ b/Engine/Source/Programs/SubmitTool/Private/Logic/PreflightService.h
@@ -37,7 +37,7 @@ public:
 		TWeakPtr<FSubmitToolServiceProvider> InServiceProvider);
 	~FPreflightService();
 
-	void RequestPreflight();
+	void RequestPreflight(const FString InTemplate = TEXT(""));
 	void FetchPreflightInfo(bool bRequeue = false, const FString& InOAuthToken = TEXT(""));
 
 	EPreflightServiceState GetState() const
@@ -64,6 +64,13 @@ public:
 	{
 		return Definition.DefaultPreflightTemplate;
 	}
+	bool GetIsHordeInformationReady() const
+	{
+		return bHordeInformationReady;
+	}
+
+	TArray<const FPreflightData*> GetTaggedPreflights() const;
+	TMap<FString, FStringFormatArg> GetFormatParameters(const FString& InTemplate = FString());
 
 	bool SelectPreflightTemplate(FPreflightTemplateDefinition& OutTemplate) const;
 
@@ -77,6 +84,7 @@ private:
 
 	void StartPreflight();
 	FString GetAdditionalTasksString(const FPreflightTemplateDefinition& InTemplate) const;
+	void RefreshStreamName();
 
 	void FetchUnlinkedPreflight(const FString& InPreflightId, bool bRequeue, const FString& InOAuthToken);
 
@@ -99,6 +107,7 @@ private:
 	// local data
 	bool bCheckShelveInstead = false;
 	bool bStopAskingTagUpdate = false;
+	bool bHordeInformationReady = false;
 	FDateTime LastRequest = FDateTime::MinValue();
 	EPreflightServiceState State;
 	FTSTicker::FDelegateHandle TickHandle;
@@ -107,8 +116,7 @@ private:
 	TMap<FString, FPreflightData> UnlinkedHordePreflights;
 
 	FString StreamName;
-
-	TMap<FString, FStringFormatArg> GetFormatParameters() const;
+	FString TemplateId;
 
 	// delete/create shelve callbacks
 	FSourceControlOperationComplete DeleteShelveCallback;
diff --git a/Engine/Source/Programs/SubmitTool/Private/Logic/TasksService.cpp b/Engine/Source/Programs/SubmitTool/Private/Logic/TasksService.cpp
--- a/Engine/Source/Programs/SubmitTool/Private/Logic/TasksService.cpp
+++ b/Engine/Source/Programs/SubmitTool/Private/Logic/TasksService.cpp
@@ -369,7 +369,7 @@ bool FTasksService::Tick(float InDeltaTime)
 				}
 			}
 
-			if(bDependenciesSatisfied && !bBlockedByGroup)
+			if(bDependenciesSatisfied && !bBlockedByGroup && Task->CanStartTask())
 			{
 				for(const FName& Group : Task->Definition->ExecutionBlockGroups)
 				{
diff --git a/Engine/Source/Programs/SubmitTool/Private/Logic/Validators/PreflightValidator.cpp b/Engine/Source/Programs/SubmitTool/Private/Logic/Validators/PreflightValidator.cpp
--- a/Engine/Source/Programs/SubmitTool/Private/Logic/Validators/PreflightValidator.cpp
+++ b/Engine/Source/Programs/SubmitTool/Private/Logic/Validators/PreflightValidator.cpp
@@ -102,39 +102,15 @@ const FString FPreflightValidator::GetPopupMessageWhenFailedText() const
 
 		TStringBuilder<1024> PreflightURLs;
 
-		for (FString PreflightId : PreflightValues)
+		for (const FPreflightData* PreflightData : PreflightService->GetTaggedPreflights())
 		{
-			if (PreflightId.Contains(TEXT("/")))
+			if (PreflightData != nullptr && (PreflightData->CachedResults.State != EPreflightState::Completed || PreflightData->CachedResults.Outcome != EPreflightOutcome::Success))
 			{
-				int32 SlashIdx;
-				PreflightId.FindLastChar(TCHAR('/'), SlashIdx);
-				PreflightId = PreflightId.RightChop(SlashIdx + 1);
-			}
-			PreflightId.TrimStartAndEndInline();
-
-			FRegexPattern Pattern = FRegexPattern(TEXT("(?:[0-9]|[a-f]){24}"), ERegexPatternFlags::CaseInsensitive);
-			FRegexMatcher regex = FRegexMatcher(Pattern, PreflightId);
-			if (regex.FindNext())
-			{
-				const TUniquePtr<FPreflightList>& PFData = PreflightService->GetPreflightData();
-				if (PFData.IsValid())
+				if (PreflightURLs.Len() != 0)
 				{
-					// Look into the linked preflights and if we don't find it check the unlinked preflights
-					const FPreflightData* FoundData = PFData->PreflightList.FindByPredicate([&PreflightId](const FPreflightData& InData) { return InData.ID == PreflightId; });
-					if (FoundData == nullptr)
-					{
-						FoundData = PreflightService->GetUnlinkedPreflights().Find(PreflightId);
-					}
-
-					if(FoundData != nullptr && (FoundData->CachedResults.State != EPreflightState::Completed || FoundData->CachedResults.Outcome != EPreflightOutcome::Success))
-					{
-						if(PreflightURLs.Len() != 0)
-						{
-							PreflightURLs << TEXT(" - ");
-						}
-						PreflightURLs << GetURLTextblock(FoundData->ID);
-					}
+					PreflightURLs << TEXT(" - ");
 				}
+				PreflightURLs << GetURLTextblock(PreflightData->ID);
 			}
 		}
 
diff --git a/Engine/Source/Programs/SubmitTool/Private/Logic/Validators/ValidatorBase.cpp b/Engine/Source/Programs/SubmitTool/Private/Logic/Validators/ValidatorBase.cpp
--- a/Engine/Source/Programs/SubmitTool/Private/Logic/Validators/ValidatorBase.cpp
+++ b/Engine/Source/Programs/SubmitTool/Private/Logic/Validators/ValidatorBase.cpp
@@ -9,7 +9,9 @@
 #include "AnalyticsEventAttribute.h"
 #include "CommandLine/CmdLineParameters.h"
 #include "Configuration/Configuration.h"
+#include "Logic/DialogFactory.h"
 #include "Logic/TagService.h"
+#include "Logic/PreflightService.h"
 #include "Logic/Services/Interfaces/ICacheDataService.h"
 #include "Logic/Services/Interfaces/ISTSourceControlService.h"
 #include "SubmitToolUtils.h"
@@ -64,6 +66,7 @@ void FValidatorBase::StartValidation()
 	RunTime = 0;
 	Start = FDateTime::UtcNow();
 	State = EValidationStates::Running;
+	FailureReason.Reset();
 	ErrorListCache.Empty();
 	WarningListCache.Empty();
 
@@ -97,6 +100,7 @@ void FValidatorBase::StartValidation()
 					LogFailure(ActivationError);
 				}
 
+				FailureReason = TEXT("Validator setup is invalid");
 				ValidationFinished(false);
 			}
 		}
@@ -104,6 +108,56 @@ void FValidatorBase::StartValidation()
 		if (bIsValidSetup)
 		{
 			TRACE_CPUPROFILER_EVENT_SCOPE(Validate);
+			const TSharedPtr<FPreflightService>& PreflightService = ServiceProvider.Pin()->GetService<FPreflightService>();
+			if (!Definition->PreflightTemplate.IsEmpty())
+			{
+				for (const FPreflightData* PreflightData : PreflightService->GetTaggedPreflights())
+				{
+					if (PreflightData != nullptr &&
+						PreflightData->TemplateId.Equals(Definition->PreflightTemplate) && 
+						PreflightData->CachedResults.State == EPreflightState::Completed && 
+						(PreflightData->CachedResults.Outcome == EPreflightOutcome::Success || PreflightData->CachedResults.Outcome == EPreflightOutcome::Warnings))
+					{
+						UE_LOG(LogSubmitTool, Log, TEXT("[%s] Preflight %s with id %s has succeeded and covers this validation."), *GetValidatorName(), *PreflightData->Name, *PreflightData->ID);
+						Skip();
+						return;
+					}
+				}
+			}
+
+			if (Definition->MinMemoryGBRequired > 0)
+			{
+				FPlatformMemoryStats MemoryStats = FPlatformMemory::GetStats();
+
+				// Only consider 80% of the page/swap in case to be conservative and allow breathing room) AvailableVirtual includes AvailablePhysical + page so underflow shouldn't be a concern
+				uint64 VirtualMemToConsider = ((MemoryStats.AvailableVirtual - MemoryStats.AvailablePhysical) * 0.8);
+				uint64 TotalMemAvailable = MemoryStats.AvailablePhysical + VirtualMemToConsider;
+
+				if ((static_cast<uint64>(Definition->MinMemoryGBRequired) * 1024 * 1024 * 1024) > TotalMemAvailable)
+				{
+					TStringBuilder<256> Message;
+					Message.Append(Definition->NotEnoughMemoryMessage);
+					Message.Appendf(TEXT("\nAvailable memory: %llu GB\nRequired memory: %d GB"), TotalMemAvailable / 1024 / 1024 / 1024, Definition->MinMemoryGBRequired);
+					const FString FinalMessage = FString::Format(Message.ToString(), PreflightService->GetFormatParameters(Definition->PreflightTemplate));
+					LogFailure(FString::Printf(TEXT("[%s] %s"), *GetValidatorName(), *FinalMessage));
+
+					TArray<FString> Buttons{ TEXT("Close")};
+					if (!Definition->PreflightTemplate.IsEmpty())
+					{
+						Buttons.Insert(TEXT("Request Preflight"), 0);
+					}
+
+					if (!Definition->PreflightTemplate.IsEmpty() && FDialogFactory::ShowDialog(FText::FromString(GetValidatorName() + TEXT(" - Not Enough Memory")), FText::FromString(FinalMessage), Buttons) == EDialogFactoryResult::FirstButton)
+					{
+						PreflightService->RequestPreflight(Definition->PreflightTemplate);
+					}
+
+					FailureReason = TEXT("Out of Memory");
+					ValidationFinished(false);
+					return;
+				}
+			}
+
 			if (IncrementallySkippedFiles.Num() != 0)
 			{
 				const FString FileList = FString::JoinBy(IncrementallySkippedFiles, TEXT("\n"), [](const FSourceControlStateRef& InFile) { return InFile->GetFilename(); });
@@ -112,6 +166,7 @@ void FValidatorBase::StartValidation()
 
 			if (!Validate(ChangelistService->GetCLDescription(), FilteredFiles, TagService->GetTagsArray()))
 			{
+				FailureReason = TEXT("Validation Failed");
 				ValidationFinished(false);
 			}
 		}
@@ -172,6 +227,16 @@ void FValidatorBase::Tick(float InDeltaTime)
 	}
 }
 
+bool FValidatorBase::CanStartTask() const
+{
+	if (!Definition->PreflightTemplate.IsEmpty())
+	{
+		return ServiceProvider.Pin()->GetService<FPreflightService>()->GetIsHordeInformationReady();
+	}
+
+	return true;
+}
+
 bool FValidatorBase::Activate()
 {
 	bIsValidSetup = true;
@@ -258,7 +323,8 @@ const TArray<FAnalyticsEventAttribute> FValidatorBase::GetTelemetryAttributes()
 		TEXT("Status"), GetHasPassed(),
 		TEXT("Runtime"), RunTime,
 		TEXT("Stream"), ServiceProvider.Pin()->GetService<ISTSourceControlService>()->GetCurrentStreamName(),
-		TEXT("PendingChangelist"), MoveTemp(ChangelistNumber)
+		TEXT("PendingChangelist"), MoveTemp(ChangelistNumber),
+		TEXT("FailureReason"), FailureReason
 	);
 }
 
diff --git a/Engine/Source/Programs/SubmitTool/Private/Logic/Validators/ValidatorBase.h b/Engine/Source/Programs/SubmitTool/Private/Logic/Validators/ValidatorBase.h
--- a/Engine/Source/Programs/SubmitTool/Private/Logic/Validators/ValidatorBase.h
+++ b/Engine/Source/Programs/SubmitTool/Private/Logic/Validators/ValidatorBase.h
@@ -161,6 +161,8 @@ public:
 		State = EValidationStates::Not_Applicable;
 	}
 
+	virtual bool CanStartTask() const;
+
 	virtual bool Activate();
 
 	virtual void InvalidateLocalFileModifications();
@@ -215,6 +217,7 @@ public:
 protected:
 	FName ValidatorNameID;
 	FString ValidatorName;
+	FString FailureReason;
 	FValidatorOptionsProvider OptionsProvider;
 	TWeakPtr<FSubmitToolServiceProvider> ServiceProvider;
 	TArray<FSourceControlStateRef> FilteredFiles;
diff --git a/Engine/Source/Programs/SubmitTool/Private/Logic/Validators/ValidatorDefinition.h b/Engine/Source/Programs/SubmitTool/Private/Logic/Validators/ValidatorDefinition.h
--- a/Engine/Source/Programs/SubmitTool/Private/Logic/Validators/ValidatorDefinition.h
+++ b/Engine/Source/Programs/SubmitTool/Private/Logic/Validators/ValidatorDefinition.h
@@ -182,12 +182,6 @@ struct FValidatorDefinition
 	UPROPERTY()
 	TArray<FString> SkipForbiddenStreams;
 
-	///
-	///	Skip is forbidden when this text if found in the CL description
-	/// 
-	UPROPERTY()
-	FString ConfigFilePath;	
-
 	///
 	///	List of Validator Ids that needs to succeed before this validator runs
 	/// 
@@ -220,7 +214,7 @@ struct FValidatorDefinition
 	/// 
 	UPROPERTY()
 	bool bUsesIncrementalCache = false;
-
+	
 	///
 	///	Additional error messages to print when this validation fails
 	/// 
@@ -260,6 +254,15 @@ struct FValidatorDefinition
 	/// 
 	UPROPERTY()
 	bool bRequireShelveOperation = false;
+
+	UPROPERTY()
+	int32 MinMemoryGBRequired = 0;
+
+	UPROPERTY()
+	FString NotEnoughMemoryMessage = TEXT("There is not enough available memory. Consider closing other programs and retrying.");
+
+	UPROPERTY()
+	FString PreflightTemplate;
 };
 
 USTRUCT()
@@ -382,6 +385,12 @@ struct FValidatorRunExecutableDefinition : public FValidatorDefinition
 	FString RegexWarningParsing = TEXT("^(?!.*(?:Display: |Log: )).*( warning |warning:).*$");
 
 	TSharedPtr<FRegexPattern> BuiltRegexWarning;
+
+	///
+	///	Additional messages to print when the exitcode is not in the allowed exit codes
+	/// 
+	UPROPERTY()
+	TArray<FString> NonZeroExitCodeErrorMessages;
 };
 
 
diff --git a/Engine/Source/Programs/SubmitTool/Private/Logic/Validators/ValidatorRunExecutable.cpp b/Engine/Source/Programs/SubmitTool/Private/Logic/Validators/ValidatorRunExecutable.cpp
--- a/Engine/Source/Programs/SubmitTool/Private/Logic/Validators/ValidatorRunExecutable.cpp
+++ b/Engine/Source/Programs/SubmitTool/Private/Logic/Validators/ValidatorRunExecutable.cpp
@@ -322,6 +322,14 @@ void FValidatorRunExecutable::OnProcessComplete(const FString& ProcessId, int32
 	const FValidatorRunExecutableDefinition* Def = GetTypedDefinition<FValidatorRunExecutableDefinition>();
 	const bool Success = Def->AllowedExitCodes.Contains(ReturnCode) && (Def->bOnlyLookAtExitCode || ProcessesData[ProcessId].ErrorList.IsEmpty());
 
+	if (!Def->AllowedExitCodes.Contains(ReturnCode))
+	{
+		for (const FString& ErrorMsg : Def->NonZeroExitCodeErrorMessages)
+		{
+			LogFailure(FString::Printf(TEXT("[%s]: %s"), *GetValidatorName(), *ErrorMsg));
+		}
+	}
+
 	if(Success)
 	{
 		UE_LOG(LogValidators, Log, TEXT("[%s (%s)]: Task process succeded (Exit code %d)"), *GetValidatorName(), *ProcessId, ReturnCode);
diff --git a/Engine/Source/Programs/SubmitTool/Private/View/Widgets/ConfirmDialogWidget.cpp b/Engine/Source/Programs/SubmitTool/Private/View/Widgets/ConfirmDialogWidget.cpp
--- a/Engine/Source/Programs/SubmitTool/Private/View/Widgets/ConfirmDialogWidget.cpp
+++ b/Engine/Source/Programs/SubmitTool/Private/View/Widgets/ConfirmDialogWidget.cpp
@@ -14,6 +14,13 @@ void SConfirmDialogWidget::Construct(const FArguments& InArgs)
 	IsBtnEnabled = InArgs._IsBtnEnabled;
 	TSharedPtr<SHorizontalBox> ButtonBox;
 	
+	FText DescriptionText = InArgs._DescriptionText;
+	FString DescriptionString = InArgs._DescriptionText.ToString();
+	if (DescriptionString.Contains(TEXT("<a")))
+	{
+		DescriptionText = FText::FromString(DescriptionString.Replace(TEXT("\'"), TEXT("\"")));
+	}
+
 	ChildSlot
 	[
 		SNew(SBorder)
@@ -27,7 +34,7 @@ void SConfirmDialogWidget::Construct(const FArguments& InArgs)
 			.Padding(5)
 			[
 				SNew(SRichTextBlock)
-				.Text(InArgs._DescriptionText)
+				.Text(DescriptionText)
 				.DecoratorStyleSet(&FAppStyle::Get())
 				+SRichTextBlock::HyperlinkDecorator(TEXT("browser"), FSlateHyperlinkRun::FOnClick::CreateLambda(
 					[](const FSlateHyperlinkRun::FMetadata& Metadata)
