Moving last shelve operation time to the cache service so it gets remembered between runs

#code-analyzed

diff --git a/Engine/Source/Programs/SubmitTool/Private/Logic/ChangelistService.cpp b/Engine/Source/Programs/SubmitTool/Private/Logic/ChangelistService.cpp
--- a/Engine/Source/Programs/SubmitTool/Private/Logic/ChangelistService.cpp
+++ b/Engine/Source/Programs/SubmitTool/Private/Logic/ChangelistService.cpp
@@ -4,6 +4,7 @@
 
 #include "Configuration/Configuration.h"
 #include "Logging/SubmitToolLog.h"
+#include "Logic/Services/Interfaces/ICacheDataService.h"
 #include "Logic/Services/Interfaces/ISTSourceControlService.h"
 #include "Models/ModelInterface.h"
 #include "CommandLine/CmdLineParameters.h"
@@ -24,14 +25,15 @@ namespace ChangelistServiceConstants {
 
 FChangelistService::FChangelistService(
 	const FGeneralParameters& InParameters,
-	const TSharedPtr<ISTSourceControlService> InPerforceService,
+	const TSharedPtr<FSubmitToolServiceProvider> InServiceProvider,
 	const FOnChangeListReadyDelegate& InCLReadyCallback,
 	const FOnChangelistRefreshDelegate& InCLRefreshCallback) :
 	Parameters(InParameters),
 	CLReadyCallback(InCLReadyCallback),
 	CLRefreshCallback(InCLRefreshCallback),
 	TickHandle(FTSTicker::GetCoreTicker().AddTicker(FTickerDelegate::CreateRaw(this, &FChangelistService::P4Tick), ChangelistServiceConstants::TickDelay)),
-	SourceControlService(InPerforceService)
+	ServiceProvider(InServiceProvider),
+	SourceControlService(InServiceProvider->GetService<ISTSourceControlService>())
 {
 	FCmdLineParameters::Get().GetValue(FSubmitToolCmdLine::P4ChangeList, CLID);
 	Init();
@@ -487,7 +489,13 @@ void FChangelistService::CreateShelvedFiles(const FSourceControlOperationComplet
 			{
 				ShelvedFilesInCL = ChangelistState->GetShelvedFilesStates();
 				ShelvedFilesDepotPaths.Reset(ShelvedFilesInCL.Num());
-				LastShelfTime = FDateTime::UtcNow();
+
+				TSharedPtr<ICacheDataService> CacheService = ServiceProvider->GetService<ICacheDataService>();
+
+				if (CacheService.IsValid())
+				{
+					CacheService->SetLastShelveOperationTime(GetCLID(), FDateTime::UtcNow());
+				}
 				UE_LOG(LogSubmitToolP4, Log, TEXT("CL %s create shelved operation succeeded"), *GetCLID());
 			}
 			else if(Result == ECommandResult::Cancelled)
@@ -751,7 +759,7 @@ const TArray<FString>& FChangelistService::GetShelvedFilesDepotPaths(bool bForce
 					{
 						if(Record.Contains(TEXT("depotFile")))
 						{
-							FilesDepotPaths.Add(Record[TEXT("depotFile")]);
+							ShelvedFilesDepotPaths.Add(Record[TEXT("depotFile")]);
 						}
 						else
 						{
diff --git a/Engine/Source/Programs/SubmitTool/Private/Logic/ChangelistService.h b/Engine/Source/Programs/SubmitTool/Private/Logic/ChangelistService.h
--- a/Engine/Source/Programs/SubmitTool/Private/Logic/ChangelistService.h
+++ b/Engine/Source/Programs/SubmitTool/Private/Logic/ChangelistService.h
@@ -11,6 +11,7 @@
 #include "Services/Interfaces/ISubmitToolService.h"
 
 struct FSourceControlResultInfo;
+class FSubmitToolServiceProvider;
 class ISTSourceControlService;
 using FSourceControlOperationRef = TSharedRef<class ISourceControlOperation, ESPMode::ThreadSafe>;
 
@@ -25,7 +26,7 @@ class FChangelistService final : public ISubmitToolService
 {
 public:
 
-	FChangelistService(const FGeneralParameters& InParameters, const TSharedPtr<ISTSourceControlService> SourceControlService, const FOnChangeListReadyDelegate& InCLReadyCallback, const FOnChangelistRefreshDelegate& InCLRefreshCallback);
+	FChangelistService(const FGeneralParameters& InParameters, const TSharedPtr<FSubmitToolServiceProvider> InServiceProvider, const FOnChangeListReadyDelegate& InCLReadyCallback, const FOnChangelistRefreshDelegate& InCLRefreshCallback);
 	~FChangelistService();
 
 	const FString GetCLID() const
@@ -67,11 +68,6 @@ public:
 		return true;
 	}
 
-	const FDateTime& GetLastShelfDateTime() const
-	{
-		return LastShelfTime;
-	}
-
 	const TArray<FSourceControlStateRef>& GetFilesInCL() const
 	{
 		return FilesInCL;
@@ -135,7 +131,7 @@ private:
 	FText OriginalDescription;
 	FString CLDescription;
 	ISourceControlProvider* CachedSCCProvider;
-	FDateTime LastShelfTime = FDateTime::MinValue();
+	const TSharedPtr<FSubmitToolServiceProvider> ServiceProvider;
 	const TSharedPtr<ISTSourceControlService> SourceControlService;
 
 	void FindInitialChangelistsAsync();
diff --git a/Engine/Source/Programs/SubmitTool/Private/Logic/Services/CacheDataService.cpp b/Engine/Source/Programs/SubmitTool/Private/Logic/Services/CacheDataService.cpp
--- a/Engine/Source/Programs/SubmitTool/Private/Logic/Services/CacheDataService.cpp
+++ b/Engine/Source/Programs/SubmitTool/Private/Logic/Services/CacheDataService.cpp
@@ -55,7 +55,6 @@ void FCacheDataService::SetIntegrationFieldValue(const FString& InCLId, const FS
 {
 	FChangelistCacheData& Data = CacheData->CLCacheData.FindOrAdd(InCLId);
 	Data.LastAccessed = FDateTime::UtcNow();
-
 	Data.IntegrationFields.FindOrAdd(InIntegrationFieldId) = InValue;
 }
 
@@ -65,6 +64,19 @@ void FCacheDataService::SaveCacheToDisk()
 	SaveToFile(FConfiguration::SubstituteAndNormalizeFilename(Parameters.CacheFile));
 }
 
+FDateTime FCacheDataService::GetLastShelveOperationTime(const FString& InCLId) const
+{
+	FChangelistCacheData& Data = CacheData->CLCacheData.FindOrAdd(InCLId);
+	Data.LastAccessed = FDateTime::UtcNow();
+	return Data.LastShelveOperation;
+}
+
+void FCacheDataService::SetLastShelveOperationTime(const FString& InCLId, const FDateTime& InDateTime)
+{
+	FChangelistCacheData& Data = CacheData->CLCacheData.FindOrAdd(InCLId);
+	Data.LastAccessed = FDateTime::UtcNow();
+	Data.LastShelveOperation = InDateTime;
+}
 
 void FCacheDataService::LoadFromFile(const FString& InFilePath)
 {
@@ -73,6 +85,12 @@ void FCacheDataService::LoadFromFile(const FString& InFilePath)
 	{
 		FString InText;
 		FArchive* File = IFileManager::Get().CreateFileReader(*InFilePath, EFileRead::FILEREAD_None);
+		if (File == nullptr)
+		{
+			UE_LOG(LogSubmitTool, Log, TEXT("Error opening cache data file %s"), *InFilePath);
+			return;
+		}
+
 		*File << InText;
 		File->Close();
 		delete File;
@@ -102,6 +120,12 @@ void FCacheDataService::SaveToFile(const FString& InFilePath) const
 	FJsonObjectConverter::UStructToJsonObjectString(*CacheData, OutputText);
 
 	FArchive* File = IFileManager::Get().CreateFileWriter(*InFilePath, EFileWrite::FILEWRITE_EvenIfReadOnly);
+	if (File == nullptr)
+	{
+		UE_LOG(LogSubmitTool, Log, TEXT("Error opening cache data file for writing %s"), *InFilePath);
+		return;
+	}
+
 	*File << OutputText;
 	File->Close();
 	delete File;
diff --git a/Engine/Source/Programs/SubmitTool/Private/Logic/Services/CacheDataService.h b/Engine/Source/Programs/SubmitTool/Private/Logic/Services/CacheDataService.h
--- a/Engine/Source/Programs/SubmitTool/Private/Logic/Services/CacheDataService.h
+++ b/Engine/Source/Programs/SubmitTool/Private/Logic/Services/CacheDataService.h
@@ -19,8 +19,11 @@ public:
 
 	virtual FString GetIntegrationFieldValue(const FString& InCLId, const FString& InIntegrationFieldId) const override { return FString(); }
 	virtual void SetIntegrationFieldValue(const FString& InCLId, const FString& InIntegrationFieldId, const FString& InValue) { }
+	virtual FDateTime GetLastShelveOperationTime(const FString& InCLId) const { return FDateTime::MinValue(); }
+	virtual void SetLastShelveOperationTime(const FString& InCLId, const FDateTime& InDateTime) { }
 
 	virtual void SaveCacheToDisk() { }
+
 };
 
 class FCacheDataService final : public ICacheDataService
@@ -36,9 +39,12 @@ public:
 
 	virtual FString GetIntegrationFieldValue(const FString& InCLId, const FString& InIntegrationFieldId) const override;
 	virtual void SetIntegrationFieldValue(const FString& InCLId, const FString& InIntegrationFieldId, const FString& InValue);
+	virtual FDateTime GetLastShelveOperationTime(const FString& InCLId) const;
+	virtual void SetLastShelveOperationTime(const FString& InCLId, const FDateTime& InDateTime);
 
 	virtual void SaveCacheToDisk();
 
+
 private:
 	void LoadFromFile(const FString& InFilePath);
 	void SaveToFile(const FString& InFilePath) const;
diff --git a/Engine/Source/Programs/SubmitTool/Private/Logic/Services/Interfaces/ICacheDataService.h b/Engine/Source/Programs/SubmitTool/Private/Logic/Services/Interfaces/ICacheDataService.h
--- a/Engine/Source/Programs/SubmitTool/Private/Logic/Services/Interfaces/ICacheDataService.h
+++ b/Engine/Source/Programs/SubmitTool/Private/Logic/Services/Interfaces/ICacheDataService.h
@@ -20,6 +20,9 @@ public:
 	virtual FString GetIntegrationFieldValue(const FString& InCLId, const FString& InIntegrationFieldId) const = 0;
 	virtual void SetIntegrationFieldValue(const FString& InCLId, const FString& InIntegrationFieldId, const FString& InValue) = 0;
 
+	virtual FDateTime GetLastShelveOperationTime(const FString& InCLId) const = 0;
+	virtual void SetLastShelveOperationTime(const FString& InCLId, const FDateTime& InDateTime) = 0;
+
 	virtual void SaveCacheToDisk() = 0;
 };
 
diff --git a/Engine/Source/Programs/SubmitTool/Private/Models/CacheData.h b/Engine/Source/Programs/SubmitTool/Private/Models/CacheData.h
--- a/Engine/Source/Programs/SubmitTool/Private/Models/CacheData.h
+++ b/Engine/Source/Programs/SubmitTool/Private/Models/CacheData.h
@@ -31,6 +31,9 @@ USTRUCT()
 struct FChangelistCacheData
 {
 	GENERATED_BODY()
+	
+	UPROPERTY()
+	FDateTime LastShelveOperation;
 
 	UPROPERTY()
 	FDateTime LastAccessed = FDateTime::UtcNow();
diff --git a/Engine/Source/Programs/SubmitTool/Private/Models/ModelInterface.cpp b/Engine/Source/Programs/SubmitTool/Private/Models/ModelInterface.cpp
--- a/Engine/Source/Programs/SubmitTool/Private/Models/ModelInterface.cpp
+++ b/Engine/Source/Programs/SubmitTool/Private/Models/ModelInterface.cpp
@@ -66,7 +66,7 @@ FModelInterface::FModelInterface(const FSubmitToolParameters& InParameters) :
 	ServiceProvider->RegisterService<FTasksService>(PresubmitOperationsService.ToSharedRef(), TEXT("PresubmitOperationsService"));
 	CredentialsService = MakeShared<FCredentialsService>(InParameters.OAuthParameters);
 	ServiceProvider->RegisterService<FCredentialsService>(CredentialsService.ToSharedRef());
-	ChangelistService = MakeShared<FChangelistService>(InParameters.GeneralParameters, SourceControlService, CLReadyCallback, CLRefreshCallback);
+	ChangelistService = MakeShared<FChangelistService>(InParameters.GeneralParameters, ServiceProvider, CLReadyCallback, CLRefreshCallback);
 	ServiceProvider->RegisterService<FChangelistService>(ChangelistService.ToSharedRef());
 	P4LockdownService = MakeShared<FP4LockdownService>(InParameters.P4LockdownParameters, ServiceProvider);
 	ServiceProvider->RegisterService<FP4LockdownService>(P4LockdownService.ToSharedRef());
@@ -1293,10 +1293,11 @@ void FModelInterface::PreValidate(const TArray<const TWeakPtr<const FValidatorBa
 				}
 			}
 
-			if (ChangelistService->GetLastShelfDateTime() < LastModifiedFile)
+			TSharedPtr<ICacheDataService> CacheService = ServiceProvider->GetService<ICacheDataService>();
+			if (ChangelistService->GetFilesInCL().Num() != ChangelistService->GetShelvedFilesInCL().Num() || !CacheService.IsValid() || CacheService->GetLastShelveOperationTime(ChangelistService->GetCLID()) < LastModifiedFile)
 			{
-				const FText TextTitle = FText::FromString(TEXT("Validation requires a shelf"));
-				const FText TextDescription = FText::FromString(TEXT("Some validations require a shelf, there are already shelved files in this changelist that differ from the local content.\n\nDo you want to delete and recreate your shelf from the latest changes in your local files for use in the preflight?"));
+				const FText TextTitle = FText::FromString(TEXT("Shelf Refresh"));
+				const FText TextDescription = FText::FromString(TEXT("Some validations require a shelf, your local files do not match the shelved ones.\n\nDo you want to refresh your shelf from the latest changes in your local files?"));
 
 				EDialogFactoryResult DialogResult = FDialogFactory::ShowDialog(TextTitle, TextDescription, TArray<FString>{ TEXT("Re-shelve files"), TEXT("Use existing Shelve")});
 
