Added additional logging for available memory

#code-analyzed

diff --git a/Engine/Source/Programs/SubmitTool/Private/Logic/Validators/ValidatorBase.cpp b/Engine/Source/Programs/SubmitTool/Private/Logic/Validators/ValidatorBase.cpp
--- a/Engine/Source/Programs/SubmitTool/Private/Logic/Validators/ValidatorBase.cpp
+++ b/Engine/Source/Programs/SubmitTool/Private/Logic/Validators/ValidatorBase.cpp
@@ -133,6 +133,8 @@ void FValidatorBase::StartValidation()
 				uint64 VirtualMemToConsider = ((MemoryStats.AvailableVirtual - MemoryStats.AvailablePhysical) * 0.8);
 				uint64 TotalMemAvailable = MemoryStats.AvailablePhysical + VirtualMemToConsider;
 
+				UE_LOG(LogSubmitToolDebug, Log, TEXT("Available memory %llu"), TotalMemAvailable / 1024 / 1024 / 1024);
+
 				if ((static_cast<uint64>(Definition->MinMemoryGBRequired) * 1024 * 1024 * 1024) > TotalMemAvailable)
 				{
 					TStringBuilder<256> Message;
