Only taking into account physical available memory when evaluating OOM

#code-analyzed

diff --git a/Engine/Source/Programs/SubmitTool/Private/Logic/Validators/ValidatorBase.cpp b/Engine/Source/Programs/SubmitTool/Private/Logic/Validators/ValidatorBase.cpp
--- a/Engine/Source/Programs/SubmitTool/Private/Logic/Validators/ValidatorBase.cpp
+++ b/Engine/Source/Programs/SubmitTool/Private/Logic/Validators/ValidatorBase.cpp
@@ -129,17 +129,14 @@ void FValidatorBase::StartValidation()
 			{
 				FPlatformMemoryStats MemoryStats = FPlatformMemory::GetStats();
 
-				// Only consider 80% of the page/swap in case to be conservative and allow breathing room) AvailableVirtual includes AvailablePhysical + page so underflow shouldn't be a concern
-				uint64 VirtualMemToConsider = ((MemoryStats.AvailableVirtual - MemoryStats.AvailablePhysical) * 0.8);
-				uint64 TotalMemAvailable = MemoryStats.AvailablePhysical + VirtualMemToConsider;
+				// Only consider Physical memory to be conservative 
+				UE_LOG(LogSubmitToolDebug, Log, TEXT("Available memory %llu"), (MemoryStats.AvailablePhysical / 1024 / 1024 / 1024));
 
-				UE_LOG(LogSubmitToolDebug, Log, TEXT("Available memory %llu"), TotalMemAvailable / 1024 / 1024 / 1024);
-
-				if ((static_cast<uint64>(Definition->MinMemoryGBRequired) * 1024 * 1024 * 1024) > TotalMemAvailable)
+				if ((static_cast<uint64>(Definition->MinMemoryGBRequired) * 1024 * 1024 * 1024) > MemoryStats.AvailablePhysical)
 				{
 					TStringBuilder<256> Message;
 					Message.Append(Definition->NotEnoughMemoryMessage);
-					Message.Appendf(TEXT("\nAvailable memory: %llu GB\nRequired memory: %d GB"), TotalMemAvailable / 1024 / 1024 / 1024, Definition->MinMemoryGBRequired);
+					Message.Appendf(TEXT("\nAvailable memory: %llu GB\nRequired memory: %d GB"), MemoryStats.AvailablePhysical / 1024 / 1024 / 1024, Definition->MinMemoryGBRequired);
 					const FString FinalMessage = FString::Format(Message.ToString(), PreflightService->GetFormatParameters(Definition->PreflightTemplate));
 					LogFailure(FString::Printf(TEXT("[%s] %s"), *GetValidatorName(), *FinalMessage));
 
