Logging to warning cache list instead of error cache list when validator is optional

#code-analyzed

diff --git a/Engine/Source/Programs/SubmitTool/Private/Logic/Validators/ValidatorBase.h b/Engine/Source/Programs/SubmitTool/Private/Logic/Validators/ValidatorBase.h
--- a/Engine/Source/Programs/SubmitTool/Private/Logic/Validators/ValidatorBase.h
+++ b/Engine/Source/Programs/SubmitTool/Private/Logic/Validators/ValidatorBase.h
@@ -250,15 +250,19 @@ protected:
 		{
 			UE_LOG(LogValidators, Error, TEXT("%s"), *FormattedMsg);
 			UE_LOG(LogValidatorsResult, Error, TEXT("%s"), *FormattedMsg);
+
+			FScopeLock Lock(&ErrorLogsMutex);
+			ErrorListCache.Add(FormattedMsg);
 		}
 		else
 		{
 			UE_LOG(LogValidators, Warning, TEXT("%s"), *FormattedMsg);
 			UE_LOG(LogValidatorsResult, Warning, TEXT("%s"), *FormattedMsg);
+
+			FScopeLock Lock(&ErrorLogsMutex);
+			WarningListCache.Add(FormattedMsg);
 		}
 
-		FScopeLock Lock(&ErrorLogsMutex);
-		ErrorListCache.Add(FormattedMsg);
 	}
 
 	virtual void Skip();
