More consistent horde readiness information to prevent validator hanging in queues
Register shelve datetime when selecting use existing shelve to avoid repeated asks

#code-analyzed

diff --git a/Engine/Source/Programs/SubmitTool/Private/Logic/PreflightService.cpp b/Engine/Source/Programs/SubmitTool/Private/Logic/PreflightService.cpp
--- a/Engine/Source/Programs/SubmitTool/Private/Logic/PreflightService.cpp
+++ b/Engine/Source/Programs/SubmitTool/Private/Logic/PreflightService.cpp
@@ -234,6 +234,7 @@ void FPreflightService::FetchPreflightInfo(bool bRequeue, const FString& InOAuth
 			QueueFetch(bRequeue, Definition.FetchPreflightEachSeconds);
 		}
 
+		bHordeInformationReady = true;
 		return;
 	}
 
@@ -427,6 +428,7 @@ void FPreflightService::FetchPreflightInfo(bool bRequeue, const FString& InOAuth
 				}
 				else
 				{
+					bHordeInformationReady = true;
 					UE_LOG(LogSubmitTool, Warning, TEXT("Couldn't obtain OAuth token login, communication with Horde is not possible."));
 				}
 			});
@@ -437,6 +439,7 @@ void FPreflightService::FetchUnlinkedPreflight(const FString& InPreflightId, boo
 {
 	if(InOAuthToken.IsEmpty())
 	{
+		bHordeInformationReady = true;
 		return;
 	}
 
@@ -482,6 +485,11 @@ void FPreflightService::FetchUnlinkedPreflight(const FString& InPreflightId, boo
 				{
 					UE_LOG(LogSubmitTool, Warning, TEXT("Unable to connect to horde. Connection error, no response."));
 				}
+
+				if (ActiveUnlinkedRequests == 0)
+				{
+					bHordeInformationReady = true;
+				}
 				return;
 			}
 
@@ -525,6 +533,7 @@ void FPreflightService::FetchUnlinkedPreflight(const FString& InPreflightId, boo
 
 			if(ActiveUnlinkedRequests == 0)
 			{
+				bHordeInformationReady = true;
 				if(OnPreflightDataUpdated.IsBound())
 				{
 					OnPreflightDataUpdated.Broadcast(HordePreflights, UnlinkedHordePreflights);
diff --git a/Engine/Source/Programs/SubmitTool/Private/Models/ModelInterface.cpp b/Engine/Source/Programs/SubmitTool/Private/Models/ModelInterface.cpp
--- a/Engine/Source/Programs/SubmitTool/Private/Models/ModelInterface.cpp
+++ b/Engine/Source/Programs/SubmitTool/Private/Models/ModelInterface.cpp
@@ -1308,6 +1308,10 @@ void FModelInterface::PreValidate(const TArray<const TWeakPtr<const FValidatorBa
 						}));
 					return;
 				}
+				else if (CacheService.IsValid())
+				{
+					CacheService->SetLastShelveOperationTime(GetCLID(), FDateTime::UtcNow());
+				}
 			}			
 		}
 		else
