Adding a way of trying to run the validator even if it's expected to run oom

diff --git a/Engine/Source/Programs/SubmitTool/Private/Logic/Validators/ValidatorBase.cpp b/Engine/Source/Programs/SubmitTool/Private/Logic/Validators/ValidatorBase.cpp
--- a/Engine/Source/Programs/SubmitTool/Private/Logic/Validators/ValidatorBase.cpp
+++ b/Engine/Source/Programs/SubmitTool/Private/Logic/Validators/ValidatorBase.cpp
@@ -140,20 +140,19 @@ void FValidatorBase::StartValidation()
 					const FString FinalMessage = FString::Format(Message.ToString(), PreflightService->GetFormatParameters(Definition->PreflightTemplate));
 					LogFailure(FString::Printf(TEXT("[%s] %s"), *GetValidatorName(), *FinalMessage));
 
-					TArray<FString> Buttons{ TEXT("Close")};
-					if (!Definition->PreflightTemplate.IsEmpty())
+					TArray<FString> Buttons{ TEXT("Request Preflight"), TEXT("Try anyway"), TEXT("Close") };
+					EDialogFactoryResult Result = FDialogFactory::ShowDialog(FText::FromString(GetValidatorName() + TEXT(" - Not Enough Memory")), FText::FromString(FinalMessage), Buttons, nullptr, [this](size_t InIdx) { return InIdx != 0 || !Definition->PreflightTemplate.IsEmpty(); });
+					if (Result != EDialogFactoryResult::SecondButton)
 					{
-						Buttons.Insert(TEXT("Request Preflight"), 0);
-					}
+						if (Result == EDialogFactoryResult::FirstButton)
+						{
+							PreflightService->RequestPreflight(Definition->PreflightTemplate);
+						}
 
-					if (!Definition->PreflightTemplate.IsEmpty() && FDialogFactory::ShowDialog(FText::FromString(GetValidatorName() + TEXT(" - Not Enough Memory")), FText::FromString(FinalMessage), Buttons) == EDialogFactoryResult::FirstButton)
-					{
-						PreflightService->RequestPreflight(Definition->PreflightTemplate);
+						FailureReason = TEXT("Out of Memory");
+						ValidationFinished(false);
+						return;
 					}
-
-					FailureReason = TEXT("Out of Memory");
-					ValidationFinished(false);
-					return;
 				}
 			}
 
diff --git a/Engine/Source/Programs/SubmitTool/Private/Logic/Validators/ValidatorDefinition.h b/Engine/Source/Programs/SubmitTool/Private/Logic/Validators/ValidatorDefinition.h
--- a/Engine/Source/Programs/SubmitTool/Private/Logic/Validators/ValidatorDefinition.h
+++ b/Engine/Source/Programs/SubmitTool/Private/Logic/Validators/ValidatorDefinition.h
@@ -259,7 +259,7 @@ struct FValidatorDefinition
 	int32 MinMemoryGBRequired = 0;
 
 	UPROPERTY()
-	FString NotEnoughMemoryMessage = TEXT("There is not enough available memory. Consider closing other programs and retrying.");
+	FString NotEnoughMemoryMessage = TEXT("There is not enough available memory for the estimated consumption. Consider closing other programs and retrying.");
 
 	UPROPERTY()
 	FString PreflightTemplate;
