Initializing last known shelve to startup time so that we don't ask for a reshelve initially if there is an existing shelf
Adding logs to the datetime comparison for shelf evaluation

diff --git a/Engine/Source/Programs/SubmitTool/Private/Models/CacheData.h b/Engine/Source/Programs/SubmitTool/Private/Models/CacheData.h
--- a/Engine/Source/Programs/SubmitTool/Private/Models/CacheData.h
+++ b/Engine/Source/Programs/SubmitTool/Private/Models/CacheData.h
@@ -33,7 +33,7 @@ struct FChangelistCacheData
 	GENERATED_BODY()
 	
 	UPROPERTY()
-	FDateTime LastShelveOperation;
+	FDateTime LastShelveOperation = FDateTime::UtcNow();
 
 	UPROPERTY()
 	FDateTime LastAccessed = FDateTime::UtcNow();
diff --git a/Engine/Source/Programs/SubmitTool/Private/Models/ModelInterface.cpp b/Engine/Source/Programs/SubmitTool/Private/Models/ModelInterface.cpp
--- a/Engine/Source/Programs/SubmitTool/Private/Models/ModelInterface.cpp
+++ b/Engine/Source/Programs/SubmitTool/Private/Models/ModelInterface.cpp
@@ -1294,12 +1294,13 @@ void FModelInterface::PreValidate(const TArray<const TWeakPtr<const FValidatorBa
 			}
 
 			TSharedPtr<ICacheDataService> CacheService = ServiceProvider->GetService<ICacheDataService>();
+			UE_LOG(LogSubmitToolDebug, Log, TEXT("Last modified file: %s, Last Shelve operation: %s"), *LastModifiedFile.ToString(), *CacheService->GetLastShelveOperationTime(ChangelistService->GetCLID()).ToString());
 			if (ChangelistService->GetFilesInCL().Num() != ChangelistService->GetShelvedFilesInCL().Num() || !CacheService.IsValid() || CacheService->GetLastShelveOperationTime(ChangelistService->GetCLID()) < LastModifiedFile)
 			{
 				const FText TextTitle = FText::FromString(TEXT("Shelf Refresh"));
 				const FText TextDescription = FText::FromString(TEXT("Some validations require a shelf, your local files do not match the shelved ones.\n\nDo you want to refresh your shelf from the latest changes in your local files?"));
 
-				EDialogFactoryResult DialogResult = FDialogFactory::ShowDialog(TextTitle, TextDescription, TArray<FString>{ TEXT("Re-shelve files"), TEXT("Use existing Shelve")});
+				EDialogFactoryResult DialogResult = FDialogFactory::ShowDialog(TextTitle, TextDescription, TArray<FString>{ TEXT("Re-shelve files"), TEXT("Use existing shelf")});
 
 				if (DialogResult == EDialogFactoryResult::FirstButton)
 				{
