Fix for doing preflight work in game thread

diff --git a/Engine/Source/Programs/SubmitTool/Private/Logic/CredentialsService.cpp b/Engine/Source/Programs/SubmitTool/Private/Logic/CredentialsService.cpp
--- a/Engine/Source/Programs/SubmitTool/Private/Logic/CredentialsService.cpp
+++ b/Engine/Source/Programs/SubmitTool/Private/Logic/CredentialsService.cpp
@@ -39,6 +39,15 @@ FCredentialsService::~FCredentialsService()
 	GetOIDCTask.Wait();
 }
 
+UE::Tasks::TTask<void> FCredentialsService::QueueWorkForToken(TFunction<void(const FString&)> InFunction)
+{
+	TRACE_CPUPROFILER_EVENT_SCOPE(FCredentialsService::QueueWorkForToken);
+	return UE::Tasks::Launch(UE_SOURCE_LOCATION, [this, InFunction]() {
+		TRACE_CPUPROFILER_EVENT_SCOPE(FCredentialsService::DoWorkForToken);
+		InFunction(OIDCToken);
+		}, GetOIDCTask, LowLevelTasks::ETaskPriority::Normal, UE::Tasks::EExtendedTaskPriority::GameThreadNormalPri);
+}
+
 void FCredentialsService::LoadKey()
 {
 	if(!IFileManager::Get().FileExists(*GetKeyFilepath()))
diff --git a/Engine/Source/Programs/SubmitTool/Private/Logic/CredentialsService.h b/Engine/Source/Programs/SubmitTool/Private/Logic/CredentialsService.h
--- a/Engine/Source/Programs/SubmitTool/Private/Logic/CredentialsService.h
+++ b/Engine/Source/Programs/SubmitTool/Private/Logic/CredentialsService.h
@@ -62,12 +62,7 @@ public:
 		return Key;
 	}
 
-	UE::Tasks::TTask<void> QueueWorkForToken(TFunction<void(const FString&)> InFunction)
-	{
-		return UE::Tasks::Launch(UE_SOURCE_LOCATION, [this, InFunction]() {
-			InFunction(OIDCToken);
-			}, GetOIDCTask);
-	}
+	UE::Tasks::TTask<void> QueueWorkForToken(TFunction<void(const FString&)> InFunction);
 
 private:
 	static TUniquePtr<FAES::FAESKey> Key;
