SubmitTool: Tweak ValidatorRunExecutable process name from a hard coded "#1" name to dynamic one "#<N>|<ExecutableName>" for improved logging.
Use the process name in a new "processes" trace region category.
Tweak all ValidatorBase logs to start with "[<ValidatorName>]".

Example output:
LogValidatorsResult: [Build UAT] Running Task
LogValidators: [Build UAT] [#2|RunUAT.bat] Task process started.
LogValidators: [Build UAT] [#2|RunUAT.bat] Completed running process #2|RunUAT.bat. Process took 3.88 sec and exited with code 0
LogValidators: [Build UAT] [#2|RunUAT.bat] Task process succeded (Exit code 0)
LogValidatorsResult: [Build UAT] Task Succeeded! (3.90 sec)

#rnx

diff --git a/Engine/Source/Programs/SubmitTool/Private/Logic/Validators/UBTValidator.cpp b/Engine/Source/Programs/SubmitTool/Private/Logic/Validators/UBTValidator.cpp
--- a/Engine/Source/Programs/SubmitTool/Private/Logic/Validators/UBTValidator.cpp
+++ b/Engine/Source/Programs/SubmitTool/Private/Logic/Validators/UBTValidator.cpp
@@ -213,7 +213,7 @@ bool FUBTValidator::Validate(const FString& InCLDescription, const TArray<FSourc
 
 	const FString AbsoluteExecutablePath = FConfiguration::SubstituteAndNormalizeFilename(ExecutablePath);
 
-	return QueueProcess(TEXT("#1"), AbsoluteExecutablePath, FinalArgs);
+	return QueueProcess(GetProcessName(AbsoluteExecutablePath), AbsoluteExecutablePath, FinalArgs);
 }
 
 bool FUBTValidator::Activate()
diff --git a/Engine/Source/Programs/SubmitTool/Private/Logic/Validators/ValidatorBase.cpp b/Engine/Source/Programs/SubmitTool/Private/Logic/Validators/ValidatorBase.cpp
--- a/Engine/Source/Programs/SubmitTool/Private/Logic/Validators/ValidatorBase.cpp
+++ b/Engine/Source/Programs/SubmitTool/Private/Logic/Validators/ValidatorBase.cpp
@@ -218,7 +218,7 @@ void FValidatorBase::Tick(float InDeltaTime)
 
 	if (Definition->TimeoutLimit > 0 && RunTime >= Definition->TimeoutLimit)
 	{
-		LogFailure(FString::Printf(TEXT("[%s]: %s"), *GetValidatorName(), TEXT("Timeout limit has been reached, cancelling task.")));
+		LogFailure(FString::Printf(TEXT("[%s] %s"), *GetValidatorName(), TEXT("Timeout limit has been reached, cancelling task.")));
 
 		StopInternalValidations();
 		SetValidationFinished(EValidationStates::Timeout);
@@ -272,13 +272,13 @@ void FValidatorBase::InvalidateLocalFileModifications()
 				{
 					if (GetIsRunning())
 					{
-						UE_LOG(LogValidators, Warning, TEXT("File %s was modified during %s run, this task needs to be run again"), *Filename, *GetValidatorName());
-						UE_LOG(LogValidatorsResult, Warning, TEXT("File %s was modified during %s run, this task needs to be run again"), *Filename, *GetValidatorName());
+						UE_LOG(LogValidators, Warning, TEXT("[%s] File %s was modified during the run, this task needs to be run again."), *GetValidatorName(), *Filename);
+						UE_LOG(LogValidatorsResult, Warning, TEXT("[%s] File %s was modified during the run, this task needs to be run again."), *GetValidatorName(), *Filename);
 					}
 					else
 					{
-						UE_LOG(LogValidators, Warning, TEXT("File %s has been modified after %s last run, this task needs to be run again."), *Filename, *GetValidatorName());
-						UE_LOG(LogValidatorsResult, Warning, TEXT("File %s has been modified after %s last run, this task needs to be run again."), *Filename, *GetValidatorName());
+						UE_LOG(LogValidators, Warning, TEXT("[%s] File %s has been modified after the last run, this task needs to be run again."), *GetValidatorName(), *Filename);
+						UE_LOG(LogValidatorsResult, Warning, TEXT("[%s] File %s has been modified after the last run, this task needs to be run again."), *GetValidatorName(), *Filename);
 					}
 
 					Invalidate();
@@ -330,7 +330,7 @@ void FValidatorBase::ValidationFinished(const bool bHasPassed)
 {
 	if (bHasPassed)
 	{
-		UE_LOG(LogValidatorsResult, Log, TEXT("[%s]: Task Succeeded! (%s)"), *GetValidatorName(), *FGenericPlatformTime::PrettyTime(RunTime));
+		UE_LOG(LogValidatorsResult, Log, TEXT("[%s] Task Succeeded! (%s)"), *GetValidatorName(), *FGenericPlatformTime::PrettyTime(RunTime));
 
 		if (Definition->bUsesIncrementalCache)
 		{
@@ -339,18 +339,18 @@ void FValidatorBase::ValidationFinished(const bool bHasPassed)
 	}
 	else if (Definition->IsRequired)
 	{
-		UE_LOG(LogValidatorsResult, Error, TEXT("[%s]: Failed on Required Task!"), *GetValidatorName());
+		UE_LOG(LogValidatorsResult, Error, TEXT("[%s] Failed on Required Task!"), *GetValidatorName());
 	}
 	else
 	{
-		UE_LOG(LogValidatorsResult, Warning, TEXT("[%s]: Failed on Optional Task!"), *GetValidatorName());
+		UE_LOG(LogValidatorsResult, Warning, TEXT("[%s] Failed on Optional Task!"), *GetValidatorName());
 	}
 
 	if (!bHasPassed)
 	{
 		for (const FString& ErrorMsg : Definition->AdditionalValidationErrorMessages)
 		{
-			LogFailure(FString::Printf(TEXT("[%s]: %s"), *GetValidatorName(), *ErrorMsg));
+			LogFailure(FString::Printf(TEXT("[%s] %s"), *GetValidatorName(), *ErrorMsg));
 		}
 	}
 
diff --git a/Engine/Source/Programs/SubmitTool/Private/Logic/Validators/ValidatorRunExecutable.cpp b/Engine/Source/Programs/SubmitTool/Private/Logic/Validators/ValidatorRunExecutable.cpp
--- a/Engine/Source/Programs/SubmitTool/Private/Logic/Validators/ValidatorRunExecutable.cpp
+++ b/Engine/Source/Programs/SubmitTool/Private/Logic/Validators/ValidatorRunExecutable.cpp
@@ -20,6 +20,12 @@
 #include "Serialization/JsonSerializer.h"
 #include "Serialization/JsonWriter.h"
 
+int32 FValidatorRunExecutable::ProcessCounter = 0;
+FString FValidatorRunExecutable::GetProcessName(FStringView ExecutablePath)
+{
+	return FString::Format(TEXT("#{0}|{1}"), { ++ProcessCounter, FPathViews::GetCleanFilename(ExecutablePath) });
+}
+
 FValidatorRunExecutable::FValidatorRunExecutable(const FName& InNameId, const FSubmitToolParameters& InParameters, TSharedRef<FSubmitToolServiceProvider> InServiceProvider, const FString& InDefinition) :
 	FValidatorBase(InNameId, InParameters, InServiceProvider, InDefinition)
 {
@@ -99,7 +105,7 @@ bool FValidatorRunExecutable::Validate(const FString& InCLDescription, const TAr
 		ExecutablePath = TypedDefinition->ExecutablePath;
 	}
 
-	return QueueProcess(TEXT("Process #1"), ExecutablePath, FinalArgs);
+	return QueueProcess(GetProcessName(ExecutablePath), ExecutablePath, FinalArgs);
 }
 
 bool FValidatorRunExecutable::Activate()
@@ -220,7 +226,7 @@ bool FValidatorRunExecutable::QueueProcess(const FString& InId, const FString& L
 
 void FValidatorRunExecutable::OnProcessOutputLine(const FString& ProcessId, const FString& Line, const EProcessOutputType& OutputType)
 {
-	const FString& FormattedMessage = FString::Printf(TEXT("[%s (%s)]: %s"), *GetValidatorName(), *ProcessId, *Line);
+	const FString& FormattedMessage = FString::Printf(TEXT("[%s] [%s] %s"), *GetValidatorName(), *ProcessId, *Line);
 	bool bIsWarning = false;
 	if(OutputType == EProcessOutputType::ProcessError || IsLineAnError(Line, bIsWarning, ProcessesData[ProcessId].bIgnoringOutputError, ProcessesData[ProcessId].bIgnoringWarningsAsErrors))
 	{
@@ -332,11 +338,11 @@ void FValidatorRunExecutable::OnProcessComplete(const FString& ProcessId, int32
 
 	if(Success)
 	{
-		UE_LOG(LogValidators, Log, TEXT("[%s (%s)]: Task process succeded (Exit code %d)"), *GetValidatorName(), *ProcessId, ReturnCode);
+		UE_LOG(LogValidators, Log, TEXT("[%s] [%s] Task process succeded (Exit code %d)"), *GetValidatorName(), *ProcessId, ReturnCode);
 	}
 	else
 	{
-		LogFailure(FString::Printf(TEXT("[%s (%s)]: Task process failed with exit code %d and %d log errors."), *GetValidatorName(), *ProcessId, ReturnCode, ProcessesData[ProcessId].ErrorList.Num()));
+		LogFailure(FString::Printf(TEXT("[%s] [%s] Task process failed with exit code %d and %d log errors."), *GetValidatorName(), *ProcessId, ReturnCode, ProcessesData[ProcessId].ErrorList.Num()));
 	}
 }
 
@@ -359,8 +365,14 @@ void FValidatorRunExecutable::Tick(float DeltaTime)
 
 			if (Process.bIsComplete)
 			{
-				const bool Success = Def->AllowedExitCodes.Contains(Process.ExitCode) && (Def->bOnlyLookAtExitCode || ProcessesData[Process.GetProcessName()].ErrorList.IsEmpty());
-				bAllSucceeded &= Success;
+				FProcessTrackingData& Data = ProcessesData.FindChecked(Process.GetProcessName());
+				if (!Data.bSuccess.IsSet())
+				{
+					bool bSuccess = Def->AllowedExitCodes.Contains(Process.ExitCode) && (Def->bOnlyLookAtExitCode || Data.ErrorList.IsEmpty());
+					Data.bSuccess.Emplace(bSuccess);
+					TRACE_END_REGION(*Process.GetProcessName());
+				}
+				bAllSucceeded &= Data.bSuccess.GetValue();
 			}
 		}
 
@@ -380,7 +392,8 @@ void FValidatorRunExecutable::Tick(float DeltaTime)
 						bAnyProcessRunning = true;
 						if (ProcessStarted)
 						{
-							UE_LOG(LogValidators, Log, TEXT("[%s (%s)]: %s"), *GetValidatorName(), *Process.GetProcessName(), TEXT("Task process started."));
+							TRACE_BEGIN_REGION(*Process.GetProcessName(), TEXT("Processes"));
+							UE_LOG(LogValidators, Log, TEXT("[%s] [%s] %s"), *GetValidatorName(), *Process.GetProcessName(), TEXT("Task process started."));
 						}
 						else
 						{
@@ -390,7 +403,7 @@ void FValidatorRunExecutable::Tick(float DeltaTime)
 									*Process.GetArgs()
 								});
 
-							LogFailure(FString::Printf(TEXT("[%s (%s)]: %s"), *GetValidatorName(), *Process.GetProcessName(),  *ErrorMessage));
+							LogFailure(FString::Printf(TEXT("[%s] [%s] %s"), *GetValidatorName(), *Process.GetProcessName(),  *ErrorMessage));
 						}
 					}
 				}
diff --git a/Engine/Source/Programs/SubmitTool/Private/Logic/Validators/ValidatorRunExecutable.h b/Engine/Source/Programs/SubmitTool/Private/Logic/Validators/ValidatorRunExecutable.h
--- a/Engine/Source/Programs/SubmitTool/Private/Logic/Validators/ValidatorRunExecutable.h
+++ b/Engine/Source/Programs/SubmitTool/Private/Logic/Validators/ValidatorRunExecutable.h
@@ -3,13 +3,15 @@
 #pragma once
 
 #include "ValidatorBase.h"
+#include "Containers/StringFwd.h"
 
 class FBuildVersion;
 
 struct FProcessTrackingData
 {
-	bool bIgnoringOutputError;
-	bool bIgnoringWarningsAsErrors;
+	TOptional<bool> bSuccess;
+	bool bIgnoringOutputError = false;
+	bool bIgnoringWarningsAsErrors = false;
 	TArray<FString> ErrorList;
 	TArray<FString> WarningList;
 };
@@ -46,9 +48,14 @@ protected:
 
 	void PrepareExecutableOptions();
 
+	static FString GetProcessName(FStringView ExecutablePath);
+
 
 protected:
 	const FString ExecutableOptions = TEXT("SelectedExecutable");
 	TMap<FString, FProcessTrackingData> ProcessesData;
 	TArray<FProcessWrapper> Processes;
+
+private:
+	static int32 ProcessCounter;
 };
