SubmitTool: SubmitToolApp: Modify the startup order to let async service tasks run while the ui is built and initialized on the main thread.
Move all UI setup code to a new FSlateUI class.
Move tick code to a new EngineLoop function.
Create the ModelInterface early to kick off any async tasks while we are taking time to build and initialize the actual UI app.
Record to telemetry that the application has started after initializing the UI to give ISTSourceControlService time to populate GetCurrentStreamName.
Add high level trace cpu scopes.

#rnx

diff --git a/Engine/Source/Programs/SubmitTool/Private/SubmitToolApp.cpp b/Engine/Source/Programs/SubmitTool/Private/SubmitToolApp.cpp
--- a/Engine/Source/Programs/SubmitTool/Private/SubmitToolApp.cpp
+++ b/Engine/Source/Programs/SubmitTool/Private/SubmitToolApp.cpp
@@ -22,6 +22,7 @@
 
 #include "Parameters/SubmitToolParameters.h"
 #include "Parameters/SubmitToolParametersBuilder.h"
+#include "ProfilingDebugging/CpuProfilerTrace.h"
 
 #include "Telemetry/TelemetryService.h"
 #include "Version/AppVersion.h"
@@ -32,9 +33,142 @@ IMPLEMENT_APPLICATION(SubmitTool, "SubmitTool");
 
 #define LOCTEXT_NAMESPACE "SubmitTool"
 
+class FSlateUI
+{
+public:
+	FSlateUI(FModelInterface* InModelInterface)
+		: ModelInterface(InModelInterface)
+		, Window(ModelInterface)
+		, MainDockTab(SNew(SDockTab))
+		, CommandList(MakeShared<FUICommandList>())
+	{ }
+	void Build(FSubmitToolUserPrefs* UserPrefs);
+
+private:
+
+	FModelInterface* ModelInterface;
+	SubmitToolWindow Window;
+	TSharedRef<SDockTab> MainDockTab;
+	TSharedRef<FUICommandList> CommandList;
+	TSharedPtr<FTabManager> TabManager;
+	TSharedPtr<SWidget> MenuWidget;
+	TSharedPtr<SDockTab> Tab;
+};
+
+TUniquePtr<FSlateUI> MakeSlateUI(FModelInterface* ModelInterface)
+{
+	TRACE_CPUPROFILER_EVENT_SCOPE(MakeSlateUI);
+	return MakeUnique<FSlateUI>(ModelInterface);
+}
+
+void FSlateUI::Build(FSubmitToolUserPrefs* UserPrefs)
+{
+	TRACE_CPUPROFILER_EVENT_SCOPE(BuildSlateUI);
+
+	// Build the slate UI for the program window
+	TabManager = FGlobalTabmanager::Get()->NewTabManager(MainDockTab);
+	// set the application name
+	FGlobalTabmanager::Get()->SetApplicationTitle(LOCTEXT("AppTitle", "SubmitTool"));
+	TabManager->SetCanDoDragOperation(false);
+
+	// set the main menu commands and interface
+	FSubmitToolCommandHandler CommandHandler;
+	CommandHandler.AddToCommandList(ModelInterface, CommandList);
+
+	FMenuBarBuilder MenuBarBuilder = FMenuBarBuilder(CommandList);
+	MenuBarBuilder.AddPullDownMenu(LOCTEXT("MainMenu", "Main Menu"), LOCTEXT("OpensMainMenu", "Opens Main Menu"), FNewMenuDelegate::CreateStatic(&FSubmitToolMenu::FillMainMenuEntries));
+#if !UE_BUILD_SHIPPING
+	MenuBarBuilder.AddPullDownMenu(LOCTEXT("Debug Tools", "Debug"), LOCTEXT("OpensDebugMenu", "Opens Debug Menu"), FNewMenuDelegate::CreateStatic(&FSubmitToolMenu::FillDebugMenuEntries));
+#endif
+
+	MenuWidget = MenuBarBuilder.MakeWidget();
+	TabManager->SetAllowWindowMenuBar(true);
+	TabManager->SetMenuMultiBox(MenuBarBuilder.GetMultiBox(), MenuWidget);
+
+	FName TabName = FName("Submit Tool");
+
+	TabManager->RegisterTabSpawner(TabName, FOnSpawnTab::CreateLambda([&Window = Window](const FSpawnTabArgs& SpawnArgs) { return Window.BuildMainTab(SpawnArgs.GetOwnerWindow()); }));
+	TabManager->RegisterDefaultTabWindowSize(TabName, FVector2D(1024, 600));
+
+	Tab = TabManager->TryInvokeTab(TabName);
+	FWindowSizeLimits WindowLimits;
+	WindowLimits.SetMinWidth(600);
+	WindowLimits.SetMinHeight(400);
+	Tab->GetParentWindow()->SetSizeLimits(WindowLimits);
+
+	if(!UserPrefs->WindowPosition.IsZero())
+	{
+		FDisplayMetrics DisplayMetrics;
+		FSlateApplicationBase::Get().GetCachedDisplayMetrics(DisplayMetrics);
+		const FPlatformRect& VirtualDisplayRect = DisplayMetrics.VirtualDisplayRect;
+
+		if(UserPrefs->WindowPosition.X >= VirtualDisplayRect.Left &&
+			UserPrefs->WindowPosition.X < VirtualDisplayRect.Right &&
+			UserPrefs->WindowPosition.Y >= VirtualDisplayRect.Top &&
+			UserPrefs->WindowPosition.Y < VirtualDisplayRect.Bottom)
+		{
+			Tab->GetParentWindow()->MoveWindowTo(UserPrefs->WindowPosition);
+		}
+	}
+
+	if(!UserPrefs->WindowSize.IsZero())
+	{
+		Tab->GetParentWindow()->Resize(UserPrefs->WindowSize);
+	}
+	else
+	{
+		Tab->GetParentWindow()->Resize(FDeprecateSlateVector2D(1024, 768));
+	}
+
+	if (UserPrefs->bWindowMaximized)
+	{
+		Tab->GetParentWindow()->Maximize();
+	}
+}
+
+static void EngineLoop(FModelInterface* ModelInterface)
+{
+	TRACE_CPUPROFILER_EVENT_SCOPE(EngineLoop);
+
+	double DeltaTime = 0.0;
+	double LastTime = FPlatformTime::Seconds();
+	const float IdealFrameTime = 1.0f / 60;
+	const float BackgroundFrameTime = 1.0f / 4;
+
+	// Loop the engine
+	while (!IsEngineExitRequested())
+	{
+		TRACE_CPUPROFILER_EVENT_SCOPE(Tick);
+
+		BeginExitIfRequested();
+
+		FTaskGraphInterface::Get().ProcessThreadUntilIdle(ENamedThreads::GameThread);
+		FTSTicker::GetCoreTicker().Tick(DeltaTime);
+		FSlateApplication::Get().PumpMessages();
+		FSlateApplication::Get().Tick();
+
+		if (IsEngineExitRequested())
+		{
+			// Dispose here so SCCProvider has time to clean up
+			ModelInterface->Dispose();
+		}
+
+		const float FrameTime = !FPlatformApplicationMisc::IsThisApplicationForeground() && (FPlatformTime::Seconds() - FSlateApplication::Get().GetLastUserInteractionTime()) > 5 ? BackgroundFrameTime : IdealFrameTime;
+		FPlatformProcess::Sleep(FMath::Max<float>(0.f, FrameTime - static_cast<float>(FPlatformTime::Seconds() - LastTime)));
+
+		DeltaTime = FPlatformTime::Seconds() - LastTime;
+		LastTime = FPlatformTime::Seconds();
+
+		UE::Stats::FStats::AdvanceFrame(false);
+		FCoreDelegates::OnEndFrame.Broadcast();
+
+		GFrameCounter++;
+	}
+}
 
 int RunSubmitTool(const TCHAR* CommandLine)
 {
+	TRACE_CPUPROFILER_EVENT_SCOPE(RunSubmitTool);
 	FTaskTagScope TaskTagScope(ETaskTag::EGameThread);
 
 	// need to make sure the cwd is correct before doing anything else
@@ -64,136 +198,47 @@ int RunSubmitTool(const TCHAR* CommandLine)
 	// Tell the module manager it may now process newly-loaded UObjects when new C++ modules are loaded
 	FModuleManager::Get().StartProcessingNewlyLoadedObjects();
 
-	// crank up a normal Slate application using the platform's standalone renderer
-	FSlateApplication::InitializeAsStandaloneApplication(GetStandardStandaloneRenderer());
-	FAppStyle::SetAppStyleSet(FSubmitToolStyle::Get());
-	FSlateApplication::InitHighDPI(true);
+	// initialize the configuration system
+	FConfiguration::Init();
 
 	// App Scope
 	{
 		UE_LOG(LogSubmitTool, Log, TEXT("Version %s"), *FAppVersion::GetVersion());
 		TUniquePtr<FSubmitToolUserPrefs> UserPrefs = FSubmitToolUserPrefs::Initialize(GetUserPrefsPath());
 
-		// initialize the configuration system
-		FConfiguration::Init();
-
 		FString ParameterFile;
 		FCmdLineParameters::Get().GetValue(FSubmitToolCmdLine::ParameterFile, ParameterFile);
 		FSubmitToolParametersBuilder ParametersBuilder(ParameterFile);
 		FSubmitToolParameters Parameters = ParametersBuilder.Build();
 
-		FTelemetryService::Init(Parameters.Telemetry);
-
 		// Create a new instance of model Interface so that UI can communicate
-		const TUniquePtr<FModelInterface> ModelInterface = MakeUnique<FModelInterface>(Parameters);
+		// Do this early to kick off any async tasks while we are taking time to build and initialize the actual UI app
+		TUniquePtr<FModelInterface> ModelInterface = MakeUnique<FModelInterface>(Parameters);
+
+		// crank up a normal Slate application using the platform's standalone renderer
+		FSlateApplication::InitializeAsStandaloneApplication(GetStandardStandaloneRenderer());
+		FAppStyle::SetAppStyleSet(FSubmitToolStyle::Get());
+		FSlateApplication::InitHighDPI(true);
+
+		// Build the slate ui
+		TUniquePtr<FSlateUI> SlateUI = MakeSlateUI(ModelInterface.Get());
+		SlateUI->Build(UserPrefs.Get());
 
 		// record that the application has started
+		// do this after initializing UI to give ISTSourceControlService time to populate GetCurrentStreamName
+		FTelemetryService::Init(Parameters.Telemetry);
 		FTelemetryService::Get()->Start(ModelInterface->GetServiceProvider()->GetService<ISTSourceControlService>()->GetCurrentStreamName());
 
-		// UI Scope
-		{
-			// Build the slate UI for the program window
-			TSharedRef<SDockTab> MainDockTab = SNew(SDockTab);
-			TSharedPtr<FTabManager> TabManager = FGlobalTabmanager::Get()->NewTabManager(MainDockTab);
-			// set the application name
-			FGlobalTabmanager::Get()->SetApplicationTitle(LOCTEXT("AppTitle", "SubmitTool"));
-			TabManager->SetCanDoDragOperation(false);
-
-			// set the main menu commands and interface
-			TSharedRef<FUICommandList> CommandList = MakeShared<FUICommandList>();
-			FSubmitToolCommandHandler CommandHandler;
-			CommandHandler.AddToCommandList(ModelInterface.Get(), CommandList);
-
-			FMenuBarBuilder MenuBarBuilder = FMenuBarBuilder(CommandList);
-			MenuBarBuilder.AddPullDownMenu(LOCTEXT("MainMenu", "Main Menu"), LOCTEXT("OpensMainMenu", "Opens Main Menu"), FNewMenuDelegate::CreateStatic(&FSubmitToolMenu::FillMainMenuEntries));
-#if !UE_BUILD_SHIPPING
-			MenuBarBuilder.AddPullDownMenu(LOCTEXT("Debug Tools", "Debug"), LOCTEXT("OpensDebugMenu", "Opens Debug Menu"), FNewMenuDelegate::CreateStatic(&FSubmitToolMenu::FillDebugMenuEntries));
-#endif
+		// Tick the application until shutdown
+		EngineLoop(ModelInterface.Get());
 
-			const TSharedRef<SWidget> MenuWidget = MenuBarBuilder.MakeWidget();
-			TabManager->SetAllowWindowMenuBar(true);
-			TabManager->SetMenuMultiBox(MenuBarBuilder.GetMultiBox(), MenuWidget);
-
-			SubmitToolWindow Window = SubmitToolWindow(ModelInterface.Get());
-			FName TabName = FName("Submit Tool");
-
-			TabManager->RegisterTabSpawner(TabName, FOnSpawnTab::CreateLambda([&Window](const FSpawnTabArgs& SpawnArgs) { return Window.BuildMainTab(SpawnArgs.GetOwnerWindow()); }));
-			TabManager->RegisterDefaultTabWindowSize(TabName, FVector2D(1024, 600));
-
-			TSharedPtr<SDockTab> Tab = TabManager->TryInvokeTab(TabName);
-			FWindowSizeLimits WindowLimits;
-			WindowLimits.SetMinWidth(600);
-			WindowLimits.SetMinHeight(400);
-			Tab->GetParentWindow()->SetSizeLimits(WindowLimits);
-
-			if(!UserPrefs->WindowPosition.IsZero())
-			{
-				FDisplayMetrics DisplayMetrics;
-				FSlateApplicationBase::Get().GetCachedDisplayMetrics(DisplayMetrics);
-				const FPlatformRect& VirtualDisplayRect = DisplayMetrics.VirtualDisplayRect;
-
-				if(UserPrefs->WindowPosition.X >= VirtualDisplayRect.Left &&
-					UserPrefs->WindowPosition.X < VirtualDisplayRect.Right &&
-					UserPrefs->WindowPosition.Y >= VirtualDisplayRect.Top &&
-					UserPrefs->WindowPosition.Y < VirtualDisplayRect.Bottom)
-				{
-					Tab->GetParentWindow()->MoveWindowTo(UserPrefs->WindowPosition);
-				}
-			}
-
-			if(!UserPrefs->WindowSize.IsZero())
-			{
-				Tab->GetParentWindow()->Resize(UserPrefs->WindowSize);
-			}
-			else
-			{
-				Tab->GetParentWindow()->Resize(FDeprecateSlateVector2D(1024, 768));
-			}
-
-
-			if (UserPrefs->bWindowMaximized)
-			{
-				Tab->GetParentWindow()->Maximize();
-			}
-
-			double DeltaTime = 0.0;
-			double LastTime = FPlatformTime::Seconds();
-			const float IdealFrameTime = 1.0f / 60;
-			const float BackgroundFrameTime = 1.0f / 4;
-
-			// Loop the engine
-			while (!IsEngineExitRequested())
-			{
-				BeginExitIfRequested();
-
-				FTaskGraphInterface::Get().ProcessThreadUntilIdle(ENamedThreads::GameThread);
-				FTSTicker::GetCoreTicker().Tick(DeltaTime);
-				FSlateApplication::Get().PumpMessages();
-				FSlateApplication::Get().Tick();
-
-				if (IsEngineExitRequested())
-				{
-					// Dispose here so SCCProvider has time to clean up
-					ModelInterface->Dispose();
-				}
-
-				const float FrameTime = !FPlatformApplicationMisc::IsThisApplicationForeground() && (FPlatformTime::Seconds() - FSlateApplication::Get().GetLastUserInteractionTime()) > 5 ? BackgroundFrameTime : IdealFrameTime;
-				FPlatformProcess::Sleep(FMath::Max<float>(0.f, FrameTime - static_cast<float>(FPlatformTime::Seconds() - LastTime)));
-
-				DeltaTime = FPlatformTime::Seconds() - LastTime;
-				LastTime = FPlatformTime::Seconds();
-
-				UE::Stats::FStats::AdvanceFrame(false);
-				FCoreDelegates::OnEndFrame.Broadcast();
-
-				GFrameCounter++;
-			}
-		}
 		// ensure all the telemetry events are flushed before unloading modules
+		TRACE_CPUPROFILER_EVENT_SCOPE(FlushTelemetry);
 		FTelemetryService::BlockFlush(5.f);
 		FTelemetryService::Shutdown();
 	}
 
+	TRACE_CPUPROFILER_EVENT_SCOPE(Shutdown);
 	FSlateApplication::Shutdown();
 
 	FEngineLoop::AppPreExit();
