New Shelving mechanism:
- Changelist service has a new EnsureShelfIsCurrent that does checks to compare shelved files vs local ones and re-shelve them when needed, for that we check:
- if there's no shelf, we just create it
- if the filepaths differ, we require a delete+shelve
- if the filepaths match, we check our cachedata to see when was the last time submit tool knew the state of the shelf was up to date, if we have any local file with a modified timestamp newer than the shelf, it needs updating
- if we don't know the last time the shelf was updated (first time submit tool is opening this CL) we fall back to the slow path of running a p4 diff

Using the new system for preflights
Only validators that require a shelf will be waiting in queue for the shelf to be ready, the rest of the validator can be started now

diff --git a/Engine/Source/Programs/SubmitTool/Private/Logic/ChangelistService.cpp b/Engine/Source/Programs/SubmitTool/Private/Logic/ChangelistService.cpp
--- a/Engine/Source/Programs/SubmitTool/Private/Logic/ChangelistService.cpp
+++ b/Engine/Source/Programs/SubmitTool/Private/Logic/ChangelistService.cpp
@@ -8,7 +8,10 @@
 #include "Logic/Services/Interfaces/ICacheDataService.h"
 #include "Logic/Services/Interfaces/ISTSourceControlService.h"
 #include "Models/ModelInterface.h"
+#include "Models/SubmitToolUserPrefs.h"
+#include "HAL/FileManagerGeneric.h"
 #include "CommandLine/CmdLineParameters.h"
+#include "SubmitToolUtils.h"
 
 #include "ISourceControlProvider.h"
 #include "ISourceControlOperation.h"
@@ -476,8 +479,9 @@ void FChangelistService::CreateShelvedFiles(const FSourceControlOperationComplet
 
 				if (CacheService.IsValid())
 				{
-					CacheService->SetLastShelveOperationTime(GetCLID(), FDateTime::UtcNow());
+					CacheService->SetLastShelfKnownUpToDate(GetCLID(), FDateTime::UtcNow());
 				}
+				bIsShelfReady = true;
 				UE_LOG(LogSubmitToolP4, Log, TEXT("CL %s create shelved operation succeeded"), *GetCLID());
 			}
 			else if(Result == ECommandResult::Cancelled)
@@ -491,6 +495,39 @@ void FChangelistService::CreateShelvedFiles(const FSourceControlOperationComplet
 	));
 }
 
+void FChangelistService::UpdateShelvedFiles(TArray<FString> InFiles, TFunction<void(bool)> OnComplete)
+{
+	TRACE_CPUPROFILER_EVENT_SCOPE(FChangelistService::UpdateShelvedFiles);
+	InFiles.Insert(TEXT("-f"), 0);
+	InFiles.Insert(TEXT("-c"), 0);
+	InFiles.Insert(CLID, 1);
+
+	SourceControlService->RunCommand(TEXT("shelve"), InFiles, FOnSCCCommandCompleteNoRet::CreateLambda(
+		[OnComplete, this](bool bSuccess, const FSCCRecordSet& InResultValues, const FSourceControlResultInfo& InResultsInfo)
+		{
+			TRACE_CPUPROFILER_EVENT_SCOPE(FChangelistService::UpdateShelvedFiles_Result);
+			UE::Tasks::Launch(UE_SOURCE_LOCATION, [OnComplete, bSuccess, this]() {
+				TRACE_CPUPROFILER_EVENT_SCOPE(FChangelistService::UpdateShelvedFiles_ResultInnerTask);
+				if (bSuccess)
+				{
+					TSharedPtr<ICacheDataService> CacheService = ServiceProvider->GetService<ICacheDataService>();
+
+					if (CacheService.IsValid())
+					{
+						CacheService->SetLastShelfKnownUpToDate(GetCLID(), FDateTime::UtcNow());
+					}
+					UE_LOG(LogSubmitToolP4, Log, TEXT("CL %s shelved files have been updated"), *GetCLID());
+					bIsShelfReady = true;
+				}
+
+				if(OnComplete != nullptr)
+				{
+					OnComplete(bSuccess);
+				}
+			}, LowLevelTasks::ETaskPriority::Normal, UE::Tasks::EExtendedTaskPriority::GameThreadNormalPri);			
+		}));
+}
+
 bool FChangelistService::P4Tick(float DeltaTime)
 {
 	for(const TSharedRef<ISourceControlOperation>& operation : ActiveP4Operations)
@@ -506,7 +543,7 @@ bool FChangelistService::IsP4OperationRunning(FName OperationName)
 {
 	if(OperationName == FName())
 	{
-		return ActiveP4Operations.Num() > 0;
+		return ActiveP4Operations.Num() > 0 || SourceControlService->IsBusy();
 	}
 
 	for(const TSharedRef<ISourceControlOperation>& Operation : ActiveP4Operations)
@@ -793,3 +830,172 @@ const TMap<FString, FString>& FChangelistService::GetOpenFilesInChangelists(bool
 	return OpenFilesInChangelists;
 }
 
+void FChangelistService::CheckShelfIsUpToDate(TFunction<void(bool, TArray<FString>)> OnComplete)
+{
+	check(OnComplete);
+	TRACE_CPUPROFILER_EVENT_SCOPE(FChangelistService::CheckShelfIsUpToDate);
+
+	TArray<FString> FilePathsOutOfDate;
+
+	if (GetFilesDepotPaths() != GetShelvedFilesDepotPaths())
+	{
+		OnComplete(false, MoveTemp(FilePathsOutOfDate));
+		return;
+	}
+	
+	// If submit tool has ever known the shelf state for this CL, use that instead of running a full diff which might take too long.
+	// This unfortunately will be incorrect in the cases that files were manually shelved in p4v and will do an unnecessary re-shelve but 
+	// still better than paying for a diff every time just to avoid this edge case
+	TSharedPtr<ICacheDataService> CacheService = ServiceProvider->GetService<ICacheDataService>();
+	if (CacheService)
+	{
+		FDateTime LastShelf = CacheService->GetLastKnownShelfUpToDate(GetCLID());
+		if (LastShelf != FDateTime::MinValue())
+		{
+			FFileManagerGeneric FileManager;
+			for (const FSourceControlStateRef& File : FilesInCL)
+			{
+				FString Filename = File->GetFilename();
+				FFileStatData FileModifiedDate = FileManager.GetStatData(*Filename);
+				if (FileModifiedDate.bIsValid && FileModifiedDate.ModificationTime > LastShelf)
+				{
+					UE_LOG(LogSubmitToolDebug, Log, TEXT("Modified file: %s at %s, Last Shelve operation: %s"), *Filename, *FileModifiedDate.ModificationTime.ToString(), *LastShelf.ToString());
+					FilePathsOutOfDate.Add(Filename);
+				}
+			}
+			
+			bIsShelfReady = FilePathsOutOfDate.Num() == 0;
+			OnComplete(bIsShelfReady, MoveTemp(FilePathsOutOfDate));
+			return;
+		}
+	}
+
+	// Only do a real p4 diff if shelf state is unknown to submit tool 
+	UE_LOG(LogSubmitToolP4, Log, TEXT("Checking Shelf diff with local files..."));
+	TArray<FString> DiffArgs = GetFilesDepotPaths();
+	for (FString& FilePath : DiffArgs)
+	{
+		FilePath.Appendf(TEXT("@=%s"), *CLID);
+	}
+	DiffArgs.Insert(TEXT("-ds"), 0);
+	DiffArgs.Insert(TEXT("-Od"), 0);
+	if (!DiffTask.IsValid() || DiffTask.IsCompleted())
+	{
+		DiffTask = SourceControlService->RunCommand(TEXT("diff"), DiffArgs, FOnSCCCommandCompleteNoRet::CreateLambda(
+			[OnComplete, this](bool bSuccess, const FSCCRecordSet& InResultValues, const FSourceControlResultInfo& InResultsInfo)
+			{
+				TRACE_CPUPROFILER_EVENT_SCOPE(FChangelistService::CheckShelfIsUpToDate_DiffResult);
+
+				TArray<FString> FilesWithChanges;
+				if (bSuccess)
+				{
+					for (const TMap<FString, FString>& Record : InResultValues)
+					{
+						if (Record.Contains(TEXT("clientFile")))
+						{
+							FilesWithChanges.Add(Record[TEXT("clientFile")]);
+						}
+						else
+						{
+							const FString Base = TEXT("clientFile");
+							size_t i = 0;
+							FString PathsKey = FString::Printf(TEXT("%s%d"), *Base, i);
+							while (Record.Contains(PathsKey))
+							{
+								FilesWithChanges.Add(Record[PathsKey]);
+								++i;
+								PathsKey = FString::Printf(TEXT("%s%d"), *Base, i);
+							}
+						}
+					}
+				}
+
+				UE::Tasks::Launch(UE_SOURCE_LOCATION, [OnComplete, FilesWithChanges = MoveTemp(FilesWithChanges), this]() {
+					bIsShelfReady = FilesWithChanges.Num() == 0;
+					// If we've checked that there are no diffs, register it to avoid future diffs
+					if (bIsShelfReady)
+					{
+						ServiceProvider->GetService<ICacheDataService>()->SetLastShelfKnownUpToDate(CLID, FDateTime::UtcNow());
+					}
+
+					OnComplete(bIsShelfReady, FilesWithChanges);
+				}, LowLevelTasks::ETaskPriority::Normal, UE::Tasks::EExtendedTaskPriority::GameThreadNormalPri);
+			}));
+	}
+}
+
+void FChangelistService::EnsureShelfIsCurrent()
+{
+	TRACE_CPUPROFILER_EVENT_SCOPE(FChangelistService::EnsureShelfIsCurrent);
+
+	bIsShelfReady = false;
+	if (GetShelvedFilesDepotPaths().Num() != 0)
+	{
+		CheckShelfIsUpToDate([this](bool bIsUpToDate, TArray<FString> InFilesOutOfDate) {
+			// if shelf is not up to date, we need to refresh it
+			if (!bIsUpToDate)
+			{
+
+				EDialogFactoryResult DialogResult = EDialogFactoryResult::Confirm;
+				if (!FSubmitToolUserPrefs::Get()->bAutoStompShelf)
+				{
+					const FText TextTitle = FText::FromString(TEXT("Shelf Refresh"));
+					FText TextDescription;
+					if (InFilesOutOfDate.IsEmpty())
+					{
+						TextDescription = FText::FromString(TEXT("Your shelf does not match your local files. Do you want to recreate your shelf with the local files in the CL?"));
+					}
+					else
+					{
+						TextDescription = FText::FromString(FString::Printf(TEXT("The following files are not up to date in the shelf. They need to be re-shelved:\n\n%s"), *FString::Join(InFilesOutOfDate, TEXT("\n"))));
+					}
+
+					DialogResult = FDialogFactory::ShowDialog(
+						TextTitle, 
+						TextDescription,
+						TArray<FString>{ TEXT("Re-shelve files"), TEXT("Cancel")},
+						FSubmitToolUtils::BuildUserPrefCheckboxUI(FSubmitToolUserPrefs::Get()->bAutoStompShelf, FText::FromString(TEXT("Always stomp a shelf when it doesn't match local state, don't ask again"))));
+				}
+
+				if (DialogResult == EDialogFactoryResult::Confirm)
+				{
+					// if we get an empty list, it means the whole shelf needs to be refreshed (it has different files)
+					if (InFilesOutOfDate.IsEmpty())
+					{
+						UE_LOG(LogSubmitToolP4, Log, TEXT("Recreating shelf"));
+						DeleteShelvedFiles(FSourceControlOperationComplete::CreateLambda([this](const FSourceControlOperationRef& Operation, ECommandResult::Type Result) {
+							CreateShelvedFiles(FSourceControlOperationComplete());
+							}));
+					}
+					else
+					{
+						// If we get a list of files out of date, we can just reshelve those instead
+						UE_LOG(LogSubmitToolP4, Log, TEXT("Refreshing shelved files that are out of date:\n%s"), *FString::Join(InFilesOutOfDate, TEXT("\n")));
+						UpdateShelvedFiles(InFilesOutOfDate);
+					}
+				}
+			}});
+			return;
+	}
+	else
+	{
+		EDialogFactoryResult DialogResult = EDialogFactoryResult::Confirm;
+		if (!FSubmitToolUserPrefs::Get()->bAutoCreateShelf)
+		{
+			const FText TextTitle = FText::FromString(TEXT("Create Shelf"));
+			const FText TextDescription = FText::FromString(TEXT("Submit tool needs a shelf. Do you want to create your shelf?"));
+		
+			DialogResult = FDialogFactory::ShowDialog(
+				TextTitle, 
+				TextDescription, 
+				TArray<FString>{ TEXT("Create shelf"), TEXT("Cancel")}, 
+				FSubmitToolUtils::BuildUserPrefCheckboxUI(FSubmitToolUserPrefs::Get()->bAutoCreateShelf, FText::FromString(TEXT("Always create a shelf when it's empty, don't ask again"))));
+		}
+
+		if (DialogResult == EDialogFactoryResult::Confirm)
+		{
+			UE_LOG(LogSubmitToolP4, Log, TEXT("Creating shelf"));
+			CreateShelvedFiles(FSourceControlOperationComplete());
+		}
+	}
+}
diff --git a/Engine/Source/Programs/SubmitTool/Private/Logic/ChangelistService.h b/Engine/Source/Programs/SubmitTool/Private/Logic/ChangelistService.h
--- a/Engine/Source/Programs/SubmitTool/Private/Logic/ChangelistService.h
+++ b/Engine/Source/Programs/SubmitTool/Private/Logic/ChangelistService.h
@@ -96,6 +96,10 @@ public:
 		return ActiveP4Operations.Num() > 0;
 	}
 
+	bool IsShelfReady() const
+	{
+		return bIsShelfReady;
+	}
 
 	void Init();
 	void Submit(const FString& InDescriptionAddendum = TEXT(""), const FSourceControlOperationComplete& OnSubmitComplete = FSourceControlOperationComplete());
@@ -105,6 +109,9 @@ public:
 	void SendCLDescriptionToP4(EConcurrency::Type Concurrency = EConcurrency::Asynchronous, FOnP4OperationCompleteDelegate InCallback = nullptr);
 	void DeleteShelvedFiles(const FSourceControlOperationComplete& OnDeleteComplete);
 	void CreateShelvedFiles(const FSourceControlOperationComplete& OnCreateComplete);
+	void UpdateShelvedFiles(TArray<FString> InFiles, TFunction<void(bool)> OnComplete = nullptr);
+	void CheckShelfIsUpToDate(TFunction<void(bool, TArray<FString>)> OnComplete);
+	void EnsureShelfIsCurrent();
 	bool P4Tick(float DeltaTime);
 
 	bool IsP4OperationRunning(FName OperationName = FName());
@@ -145,7 +152,8 @@ private:
 	TArray<FString> ShelvedFilesDepotPaths;
 	// Map filled once by GetOpenFilesInChangelists which is more efficient than a FUpdatePendingChangelistsStatus operation for all states for all pending changelists.
 	TMap<FString, FString> OpenFilesInChangelists;
-
+	bool bIsShelfReady = false;
+	UE::Tasks::FTask DiffTask;
 	void FindInitialChangelistsAsync();
 
 	void WaitForInitialP4RunTasks();
diff --git a/Engine/Source/Programs/SubmitTool/Private/Logic/PreflightService.cpp b/Engine/Source/Programs/SubmitTool/Private/Logic/PreflightService.cpp
--- a/Engine/Source/Programs/SubmitTool/Private/Logic/PreflightService.cpp
+++ b/Engine/Source/Programs/SubmitTool/Private/Logic/PreflightService.cpp
@@ -24,6 +24,7 @@
 #include "Models/SubmitToolUserPrefs.h"
 #include "Telemetry/TelemetryService.h"
 #include "Parameters/SubmitToolParameters.h"
+#include "SubmitToolUtils.h"
 
 FPreflightService::FPreflightService(
 	const FHordeParameters& InSettings,
@@ -38,9 +39,6 @@ FPreflightService::FPreflightService(
 {
 	TickHandle = FTSTicker::GetCoreTicker().AddTicker(FTickerDelegate::CreateRaw(this, &FPreflightService::Tick));
 
-	DeleteShelveCallback = FSourceControlOperationComplete::CreateRaw(this, &FPreflightService::OnDeleteShelveOperationComplete);
-	CreateShelveCallback = FSourceControlOperationComplete::CreateRaw(this, &FPreflightService::OnCreateShelveOperationComplete);
-
 	PreflightTag = ServiceProvider.Pin()->GetService<FTagService>()->GetTagOfSubtype(TEXT("preflight"));
 }
 
@@ -60,59 +58,13 @@ bool FPreflightService::Tick(float DeltaTime)
 			// Do nothing, wait for someone to press the "Start" preflight button
 			break;
 
-			//////////////////////////////////////////////////
-		case EPreflightServiceState::RequestDeleteShelve:
-
-			if(ChangelistService->GetFilesInCL().Num() == 0)
-			{
-				LastErrorMessage = TEXT("Missing local files!  We can't update the shelved files!");
-				State = EPreflightServiceState::Error;
-			}
-			else if(ChangelistService->HasShelvedFiles())
-			{
-				EDialogFactoryResult DialogResult = static_cast<EDialogFactoryResult>(ShowRecreateShelveDialog());
-				if(DialogResult == EDialogFactoryResult::FirstButton)
-				{
-					UE_LOG(LogSubmitTool, Log, TEXT("Preflight: Deleting shelved files"));
-					State = EPreflightServiceState::WaitingForDeleteShelve;
-					bCheckShelveInstead = false;
-					ChangelistService->DeleteShelvedFiles(DeleteShelveCallback);
-				}
-				else if(DialogResult == EDialogFactoryResult::SecondButton)
-				{
-					UE_LOG(LogSubmitTool, Log, TEXT("Preflight: Using existing shelved files"));
-					bCheckShelveInstead = true;
-					State = EPreflightServiceState::StartPreflight;
-				}
-				else
-				{
-					UE_LOG(LogSubmitTool, Log, TEXT("Preflight: Request cancelled"));
-					State = EPreflightServiceState::Idle;
-				}
-			}
-			else
+		case EPreflightServiceState::WaitingForShelf:
+			if (ChangelistService->IsShelfReady())
 			{
-				State = EPreflightServiceState::RequestCreateShelve;
+				State = EPreflightServiceState::StartPreflight;
 			}
 			break;
 
-			//////////////////////////////////////////////////
-		case EPreflightServiceState::WaitingForDeleteShelve:
-			// Do nothing, wait for the delete shelve callback
-			break;
-
-			//////////////////////////////////////////////////
-		case EPreflightServiceState::RequestCreateShelve:
-			UE_LOG(LogSubmitTool, Log, TEXT("Preflight: Creating shelved files..."));
-			State = EPreflightServiceState::WaitingForCreateShelve;
-			ChangelistService->CreateShelvedFiles(CreateShelveCallback);
-			break;
-
-			//////////////////////////////////////////////////
-		case EPreflightServiceState::WaitingForCreateShelve:
-			// Do nothing, wait for the create shelve callback
-			break;
-
 			//////////////////////////////////////////////////
 		case EPreflightServiceState::StartPreflight:
 			StartPreflight();
@@ -167,7 +119,11 @@ void FPreflightService::RequestPreflight(const FString InTemplate)
 		{
 			TemplateId = InTemplate;
 			UE_LOG(LogSubmitTool, Log, TEXT("Preflight: Requesting..."));
-			State = EPreflightServiceState::RequestDeleteShelve;
+			
+			TSharedPtr<FChangelistService> ChangelistService = ServiceProvider.Pin()->GetService<FChangelistService>();
+			
+			ChangelistService->EnsureShelfIsCurrent();
+			State = EPreflightServiceState::WaitingForShelf;
 		}
 		else
 		{
@@ -549,7 +505,6 @@ void FPreflightService::FetchUnlinkedPreflight(const FString& InPreflightId, boo
 	UnlinkedPFRequest->ProcessRequest();
 }
 
-
 void FPreflightService::StartPreflight()
 {
 	const TArray<FString>& ShelvedPaths = ServiceProvider.Pin()->GetService<FChangelistService>()->GetShelvedFilesDepotPaths(true);
@@ -589,48 +544,6 @@ void FPreflightService::StartPreflight()
 	}
 }
 
-void FPreflightService::OnDeleteShelveOperationComplete(const FSourceControlOperationRef& Operation, ECommandResult::Type Result)
-{
-	if(State == EPreflightServiceState::WaitingForDeleteShelve)
-	{
-		if(Result == ECommandResult::Type::Succeeded)
-		{
-			State = EPreflightServiceState::RequestCreateShelve;
-		}
-		else
-		{
-			LastErrorMessage = TEXT("Unable to delete shelve for preflight");
-			State = EPreflightServiceState::Error;
-		}
-	}
-	else
-	{
-		LastErrorMessage = TEXT("Received delete shelve callback when not waiting for it");
-		State = EPreflightServiceState::Error;
-	}
-}
-
-void FPreflightService::OnCreateShelveOperationComplete(const FSourceControlOperationRef& Operation, ECommandResult::Type Result)
-{
-	if(State == EPreflightServiceState::WaitingForCreateShelve)
-	{
-		if(Result == ECommandResult::Type::Succeeded)
-		{
-			State = EPreflightServiceState::StartPreflight;
-		}
-		else
-		{
-			LastErrorMessage = TEXT("Unable to create shelve for preflight");
-			State = EPreflightServiceState::Error;
-		}
-	}
-	else
-	{
-		LastErrorMessage = TEXT("Received create shelve callback when not waiting for it");
-		State = EPreflightServiceState::Error;
-	}
-}
-
 TArray<const FPreflightData*> FPreflightService::GetTaggedPreflights() const
 {
 	TArray<const FPreflightData*> FoundPreflights;
@@ -778,19 +691,6 @@ void FPreflightService::RefreshStreamName()
 	}
 }
 
-EDialogFactoryResult FPreflightService::ShowRecreateShelveDialog() const
-{
-	if(ModelInterface->GetMainTab().IsValid() && ModelInterface->GetMainTab().Pin()->GetParentWindow().IsValid())
-	{
-		ModelInterface->GetMainTab().Pin()->GetParentWindow()->DrawAttention(FWindowDrawAttentionParameters());
-	}
-
-	const FText TextTitle = FText::FromString(FString::Printf(TEXT("Preflight CL %s: Recreate shelved files?"), *ServiceProvider.Pin()->GetService<FChangelistService>()->GetCLID()));
-	const FText TextDescription = FText::FromString(TEXT("There are already shelved files in this changelist.\n\nDo you want to delete and recreate your shelf from the latest changes in your local files for use in the preflight?"));
-
-	return FDialogFactory::ShowDialog(TextTitle, TextDescription, TArray<FString>{ TEXT("Re-shelve files"), TEXT("Use existing Shelve"), TEXT("Cancel") });
-}
-
 EDialogFactoryResult FPreflightService::ShowUpdatePreflightTagDialog() const
 {
 	if(ModelInterface->GetMainTab().IsValid() && ModelInterface->GetMainTab().Pin()->GetParentWindow().IsValid())
@@ -801,33 +701,5 @@ EDialogFactoryResult FPreflightService::ShowUpdatePreflightTagDialog() const
 	const FText TextTitle = FText::FromString(FString::Printf(TEXT("Preflight CL %s: Newer preflight available"), *ServiceProvider.Pin()->GetService<FChangelistService>()->GetCLID()));
 	const FText TextDescription = FText::FromString(FString::Printf(TEXT("There is a newer preflight for this changelist:\n<a id=\"browser\" style=\"Hyperlink\" href=\"%sjob/%s\">%s - %s</>\n\nDo you want to update the #preflight tag?"), *Definition.HordeServerAddress, *HordePreflights->PreflightList[0].ID, *HordePreflights->PreflightList[0].Name, *HordePreflights->PreflightList[0].ID));
 
-	TSharedPtr<SHorizontalBox> AutoUpdate = SNew(SHorizontalBox)
-		+ SHorizontalBox::Slot()
-		.AutoWidth()
-		.VAlign(VAlign_Center)
-		[
-			SNew(SCheckBox)
-				.IsChecked_Lambda([]() { return FSubmitToolUserPrefs::Get()->bAutoUpdatePreflight ? ECheckBoxState::Checked : ECheckBoxState::Unchecked; })
-				.OnCheckStateChanged_Lambda([](ECheckBoxState InNewState)
-					{
-						FSubmitToolUserPrefs::Get()->bAutoUpdatePreflight = InNewState == ECheckBoxState::Checked;
-					})
-		]
-		+ SHorizontalBox::Slot()
-		.AutoWidth()
-		.VAlign(VAlign_Center)
-		[
-			SNew(SButton)
-				.ButtonStyle(FAppStyle::Get(), "InvisibleButton")
-				.IsFocusable(false)
-				.OnClicked_Lambda([this]() { FSubmitToolUserPrefs::Get()->bAutoUpdatePreflight = !FSubmitToolUserPrefs::Get()->bAutoUpdatePreflight; return FReply::Handled(); })
-				[
-					SNew(STextBlock)
-						.Justification(ETextJustify::Left)
-						.MinDesiredWidth(60)
-						.Text(FText::FromString(TEXT("Always update, Don't ask again")))
-				]
-		];
-
-	return FDialogFactory::ShowDialog(TextTitle, TextDescription, TArray<FString>{ TEXT("Update Tag Value"), TEXT("Cancel") }, AutoUpdate);
+	return FDialogFactory::ShowDialog(TextTitle, TextDescription, TArray<FString>{ TEXT("Update Tag Value"), TEXT("Cancel") }, FSubmitToolUtils::BuildUserPrefCheckboxUI(FSubmitToolUserPrefs::Get()->bAutoUpdatePreflight, FText::FromString(TEXT("Always update, Don't ask again"))));
 }
diff --git a/Engine/Source/Programs/SubmitTool/Private/Logic/PreflightService.h b/Engine/Source/Programs/SubmitTool/Private/Logic/PreflightService.h
--- a/Engine/Source/Programs/SubmitTool/Private/Logic/PreflightService.h
+++ b/Engine/Source/Programs/SubmitTool/Private/Logic/PreflightService.h
@@ -18,10 +18,7 @@ class IHttpRequest;
 enum class EPreflightServiceState
 {
 	Idle,
-	RequestDeleteShelve,
-	WaitingForDeleteShelve,
-	RequestCreateShelve,
-	WaitingForCreateShelve,
+	WaitingForShelf,
 	StartPreflight,
 	Error,
 };
@@ -88,12 +85,7 @@ private:
 
 	void FetchUnlinkedPreflight(const FString& InPreflightId, bool bRequeue, const FString& InOAuthToken);
 
-	EDialogFactoryResult ShowRecreateShelveDialog() const;
 	EDialogFactoryResult ShowUpdatePreflightTagDialog() const;
-
-	void OnDeleteShelveOperationComplete(const FSourceControlOperationRef& Operation, ECommandResult::Type Result);
-	void OnCreateShelveOperationComplete(const FSourceControlOperationRef& Operation, ECommandResult::Type Result);
-
 	bool Tick(float DeltaTime);
 
 	// Definitions from the ini
@@ -118,10 +110,6 @@ private:
 	FString StreamName;
 	FString TemplateId;
 
-	// delete/create shelve callbacks
-	FSourceControlOperationComplete DeleteShelveCallback;
-	FSourceControlOperationComplete CreateShelveCallback;
-
 	// get stream data process
 	TUniquePtr<FProcessWrapper> GetStreamDataProcess;
 	FOnCompleted OnGetStreamDataCompletedCallback;
diff --git a/Engine/Source/Programs/SubmitTool/Private/Logic/Services/CacheDataService.cpp b/Engine/Source/Programs/SubmitTool/Private/Logic/Services/CacheDataService.cpp
--- a/Engine/Source/Programs/SubmitTool/Private/Logic/Services/CacheDataService.cpp
+++ b/Engine/Source/Programs/SubmitTool/Private/Logic/Services/CacheDataService.cpp
@@ -64,18 +64,18 @@ void FCacheDataService::SaveCacheToDisk()
 	SaveToFile(FConfiguration::SubstituteAndNormalizeFilename(Parameters.CacheFile));
 }
 
-FDateTime FCacheDataService::GetLastShelveOperationTime(const FString& InCLId) const
+FDateTime FCacheDataService::GetLastKnownShelfUpToDate(const FString& InCLId) const
 {
 	FChangelistCacheData& Data = CacheData->CLCacheData.FindOrAdd(InCLId);
 	Data.LastAccessed = FDateTime::UtcNow();
-	return Data.LastShelveOperation;
+	return Data.LastShelfIsUpToDate;
 }
 
-void FCacheDataService::SetLastShelveOperationTime(const FString& InCLId, const FDateTime& InDateTime)
+void FCacheDataService::SetLastShelfKnownUpToDate(const FString& InCLId, const FDateTime& InDateTime)
 {
 	FChangelistCacheData& Data = CacheData->CLCacheData.FindOrAdd(InCLId);
 	Data.LastAccessed = FDateTime::UtcNow();
-	Data.LastShelveOperation = InDateTime;
+	Data.LastShelfIsUpToDate = InDateTime;
 }
 
 void FCacheDataService::LoadFromFile(const FString& InFilePath)
diff --git a/Engine/Source/Programs/SubmitTool/Private/Logic/Services/CacheDataService.h b/Engine/Source/Programs/SubmitTool/Private/Logic/Services/CacheDataService.h
--- a/Engine/Source/Programs/SubmitTool/Private/Logic/Services/CacheDataService.h
+++ b/Engine/Source/Programs/SubmitTool/Private/Logic/Services/CacheDataService.h
@@ -19,8 +19,8 @@ public:
 
 	virtual FString GetIntegrationFieldValue(const FString& InCLId, const FString& InIntegrationFieldId) const override { return FString(); }
 	virtual void SetIntegrationFieldValue(const FString& InCLId, const FString& InIntegrationFieldId, const FString& InValue) { }
-	virtual FDateTime GetLastShelveOperationTime(const FString& InCLId) const { return FDateTime::MinValue(); }
-	virtual void SetLastShelveOperationTime(const FString& InCLId, const FDateTime& InDateTime) { }
+	virtual FDateTime GetLastKnownShelfUpToDate(const FString& InCLId) const { return FDateTime::MinValue(); }
+	virtual void SetLastShelfKnownUpToDate(const FString& InCLId, const FDateTime& InDateTime) { }
 
 	virtual void SaveCacheToDisk() { }
 
@@ -39,8 +39,8 @@ public:
 
 	virtual FString GetIntegrationFieldValue(const FString& InCLId, const FString& InIntegrationFieldId) const override;
 	virtual void SetIntegrationFieldValue(const FString& InCLId, const FString& InIntegrationFieldId, const FString& InValue);
-	virtual FDateTime GetLastShelveOperationTime(const FString& InCLId) const;
-	virtual void SetLastShelveOperationTime(const FString& InCLId, const FDateTime& InDateTime);
+	virtual FDateTime GetLastKnownShelfUpToDate(const FString& InCLId) const;
+	virtual void SetLastShelfKnownUpToDate(const FString& InCLId, const FDateTime& InDateTime);
 
 	virtual void SaveCacheToDisk();
 
diff --git a/Engine/Source/Programs/SubmitTool/Private/Logic/Services/Interfaces/ICacheDataService.h b/Engine/Source/Programs/SubmitTool/Private/Logic/Services/Interfaces/ICacheDataService.h
--- a/Engine/Source/Programs/SubmitTool/Private/Logic/Services/Interfaces/ICacheDataService.h
+++ b/Engine/Source/Programs/SubmitTool/Private/Logic/Services/Interfaces/ICacheDataService.h
@@ -20,8 +20,8 @@ public:
 	virtual FString GetIntegrationFieldValue(const FString& InCLId, const FString& InIntegrationFieldId) const = 0;
 	virtual void SetIntegrationFieldValue(const FString& InCLId, const FString& InIntegrationFieldId, const FString& InValue) = 0;
 
-	virtual FDateTime GetLastShelveOperationTime(const FString& InCLId) const = 0;
-	virtual void SetLastShelveOperationTime(const FString& InCLId, const FDateTime& InDateTime) = 0;
+	virtual FDateTime GetLastKnownShelfUpToDate(const FString& InCLId) const = 0;
+	virtual void SetLastShelfKnownUpToDate(const FString& InCLId, const FDateTime& InDateTime) = 0;
 
 	virtual void SaveCacheToDisk() = 0;
 };
diff --git a/Engine/Source/Programs/SubmitTool/Private/Logic/Services/Interfaces/ISTSourceControlService.h b/Engine/Source/Programs/SubmitTool/Private/Logic/Services/Interfaces/ISTSourceControlService.h
--- a/Engine/Source/Programs/SubmitTool/Private/Logic/Services/Interfaces/ISTSourceControlService.h
+++ b/Engine/Source/Programs/SubmitTool/Private/Logic/Services/Interfaces/ISTSourceControlService.h
@@ -121,7 +121,7 @@ public:
 	virtual const FString GetCurrentStreamName() = 0;
 	virtual const size_t GetDepotStreamLength(const FString& InDepotName) = 0;
 	virtual const FAuthTicket& GetAuthTicket() = 0;
-
+	virtual bool IsBusy() const = 0;
 
 	template<typename T>
 	UE::Tasks::TTask<FSCCResultRetVal<T>> RunCommand(const FString& InCommand, const TArray<FString>& InAdditionalArgs, FOnSCCCommandCompleteRetVal<T> InCompleteCallback, TArray<FSharedBuffer>* OutData = nullptr)
diff --git a/Engine/Source/Programs/SubmitTool/Private/Logic/Services/SourceControl/SubmitToolPerforce.cpp b/Engine/Source/Programs/SubmitTool/Private/Logic/Services/SourceControl/SubmitToolPerforce.cpp
--- a/Engine/Source/Programs/SubmitTool/Private/Logic/Services/SourceControl/SubmitToolPerforce.cpp
+++ b/Engine/Source/Programs/SubmitTool/Private/Logic/Services/SourceControl/SubmitToolPerforce.cpp
@@ -863,6 +863,11 @@ const FAuthTicket& FSubmitToolPerforce::GetAuthTicket()
 	return P4Ticket;
 }
 
+bool FSubmitToolPerforce::IsBusy() const
+{
+	return ActiveCommands != 0;
+}
+
 bool FSubmitToolPerforce::RunP4Command(const FString& InCommand, const TArray<FString>& InAdditionalArgs, TArray<TMap<FString, FString>>& OutResultValues, FSourceControlResultInfo& OutResults, EP4ClientUserFlags InCmdFlags, TArray<FSharedBuffer>* OutData)
 {
 	TRACE_CPUPROFILER_EVENT_SCOPE_TEXT(*FString::Printf(TEXT("FSubmitToolPerforce::RunP4Command_%s"), *InCommand));
@@ -886,6 +891,7 @@ bool FSubmitToolPerforce::RunP4Command(const FString& InCommand, const TArray<FS
 
 		++Attempts;
 	}
+	ActiveCommands++;
 
 	ClientApi& P4Client = Connection->GetConnection();
 	FString FullCommand = TEXT("p4 -p ");
@@ -968,5 +974,7 @@ bool FSubmitToolPerforce::RunP4Command(const FString& InCommand, const TArray<FS
 	}
 	delete[] ArgV;
 
+
+	ActiveCommands--;
 	return !OutResults.HasErrors();
 }
diff --git a/Engine/Source/Programs/SubmitTool/Private/Logic/Services/SourceControl/SubmitToolPerforce.h b/Engine/Source/Programs/SubmitTool/Private/Logic/Services/SourceControl/SubmitToolPerforce.h
--- a/Engine/Source/Programs/SubmitTool/Private/Logic/Services/SourceControl/SubmitToolPerforce.h
+++ b/Engine/Source/Programs/SubmitTool/Private/Logic/Services/SourceControl/SubmitToolPerforce.h
@@ -104,6 +104,7 @@ public:
 	virtual const FString GetCurrentStreamName() override;
 	virtual const size_t GetDepotStreamLength(const FString& InDepotName) override;
 	virtual const FAuthTicket& GetAuthTicket() override;
+	virtual bool IsBusy() const override;
 			
 protected:
 	virtual bool RunCommandInternal(const FString& InCommand, const TArray<FString>& InAdditionalArgs, FSCCRecordSet& OutResultValues, FSourceControlResultInfo& OutResultInfo, TArray<FSharedBuffer>* OutData = nullptr) override;
@@ -123,6 +124,7 @@ private:
 	FOnUsersGet OnUsersGetCallbacks;
 	FOnGroupsGet OnGroupsGetCallbacks;
 	FOnUsersAndGroupsGet OnUsersAndGroupsGetCallbacks;
+	std::atomic<uint64> ActiveCommands = 0;
 	
 	FTSTicker::FDelegateHandle TickHandle;
 	TUniquePtr<ISourceControlProvider> SCCProvider;
diff --git a/Engine/Source/Programs/SubmitTool/Private/Logic/TasksService.cpp b/Engine/Source/Programs/SubmitTool/Private/Logic/TasksService.cpp
--- a/Engine/Source/Programs/SubmitTool/Private/Logic/TasksService.cpp
+++ b/Engine/Source/Programs/SubmitTool/Private/Logic/TasksService.cpp
@@ -206,7 +206,7 @@ void FTasksService::InvalidateForChanges(ETaskArea InChangeType)
 		const TSharedPtr<FValidatorBase>& Task = TaskPair.Value;
 		if((Task->Definition->TaskArea & InChangeType) != ETaskArea::None)
 		{
-			Task->Invalidate();
+			Task->Invalidate(true);
 		}
 	}
 }
diff --git a/Engine/Source/Programs/SubmitTool/Private/Logic/Validators/ValidatorBase.cpp b/Engine/Source/Programs/SubmitTool/Private/Logic/Validators/ValidatorBase.cpp
--- a/Engine/Source/Programs/SubmitTool/Private/Logic/Validators/ValidatorBase.cpp
+++ b/Engine/Source/Programs/SubmitTool/Private/Logic/Validators/ValidatorBase.cpp
@@ -138,12 +138,12 @@ void FValidatorBase::StartValidation()
 					Message.Append(Definition->NotEnoughMemoryMessage);
 					Message.Appendf(TEXT("\nAvailable memory: %llu GB\nRequired memory: %d GB"), MemoryStats.AvailablePhysical / 1024 / 1024 / 1024, Definition->MinMemoryGBRequired);
 					const FString FinalMessage = FString::Format(Message.ToString(), PreflightService->GetFormatParameters(Definition->PreflightTemplate));
-					LogFailure(FString::Printf(TEXT("[%s] %s"), *GetValidatorName(), *FinalMessage));
 
 					TArray<FString> Buttons{ TEXT("Request Preflight"), TEXT("Try anyway"), TEXT("Close") };
 					EDialogFactoryResult Result = FDialogFactory::ShowDialog(FText::FromString(GetValidatorName() + TEXT(" - Not Enough Memory")), FText::FromString(FinalMessage), Buttons, nullptr, [this](size_t InIdx) { return InIdx != 0 || !Definition->PreflightTemplate.IsEmpty(); });
 					if (Result != EDialogFactoryResult::SecondButton)
 					{
+						LogFailure(FString::Printf(TEXT("[%s] %s"), *GetValidatorName(), *FinalMessage));
 						if (Result == EDialogFactoryResult::FirstButton)
 						{
 							PreflightService->RequestPreflight(Definition->PreflightTemplate);
@@ -227,12 +227,18 @@ void FValidatorBase::Tick(float InDeltaTime)
 
 bool FValidatorBase::CanStartTask() const
 {
+	bool bCanRun = true;
+	if (Definition->bRequireShelveOperation)
+	{		
+		bCanRun &= ServiceProvider.Pin()->GetService<FChangelistService>()->IsShelfReady();
+	}
+
 	if (!Definition->PreflightTemplate.IsEmpty())
 	{
-		return ServiceProvider.Pin()->GetService<FPreflightService>()->GetIsHordeInformationReady();
+		bCanRun &= ServiceProvider.Pin()->GetService<FPreflightService>()->GetIsHordeInformationReady();
 	}
 
-	return true;
+	return bCanRun;
 }
 
 bool FValidatorBase::Activate()
@@ -294,20 +300,39 @@ const FString FValidatorBase::GetStatusText() const
 	const FString StateStr = StaticEnum<EValidationStates>()->GetNameStringByValue(static_cast<int64>(State))
 		.Replace(TEXT("_"), TEXT(" "));
 
-	if (State == EValidationStates::Skipped || State == EValidationStates::Not_Run || State == EValidationStates::Not_Applicable)
+	FString AdditionalInformation;
+
+	if (State == EValidationStates::Queued)
 	{
-		return StateStr;
+		if (Definition->bRequireShelveOperation && !ServiceProvider.Pin()->GetService<FChangelistService>()->IsShelfReady())
+		{
+			AdditionalInformation = TEXT("Waiting Shelf");
+		}
+
+		if (!Definition->PreflightTemplate.IsEmpty() && !ServiceProvider.Pin()->GetService<FPreflightService>()->GetIsHordeInformationReady())
+		{
+			AdditionalInformation = TEXT("Waiting Horde");
+		}
 	}
 
 	// do not clutter the UI with uninteresting information
-	if (RunTime < 0.5f)
+	if (State == EValidationStates::Failed || State == EValidationStates::Valid || State == EValidationStates::Running || State == EValidationStates::Timeout)
+	{
+		if (RunTime >= 0.5f)
+		{
+			AdditionalInformation = FGenericPlatformTime::PrettyTime(RunTime);
+		}		
+	}
+
+	if (AdditionalInformation.IsEmpty())
 	{
 		return StateStr;
 	}
+	else
+	{
+		return FString::Printf(TEXT("%s (%s)"), *StateStr,	*AdditionalInformation);
+	}
 
-	return FString::Printf(TEXT("%s (%s)"),
-		*StateStr,
-		*FGenericPlatformTime::PrettyTime(RunTime));
 }
 
 const TArray<FAnalyticsEventAttribute> FValidatorBase::GetTelemetryAttributes() const
diff --git a/Engine/Source/Programs/SubmitTool/Private/Models/CacheData.h b/Engine/Source/Programs/SubmitTool/Private/Models/CacheData.h
--- a/Engine/Source/Programs/SubmitTool/Private/Models/CacheData.h
+++ b/Engine/Source/Programs/SubmitTool/Private/Models/CacheData.h
@@ -33,7 +33,7 @@ struct FChangelistCacheData
 	GENERATED_BODY()
 	
 	UPROPERTY()
-	FDateTime LastShelveOperation = FDateTime::UtcNow();
+	FDateTime LastShelfIsUpToDate = FDateTime::MinValue();
 
 	UPROPERTY()
 	FDateTime LastAccessed = FDateTime::UtcNow();
diff --git a/Engine/Source/Programs/SubmitTool/Private/Models/ModelInterface.cpp b/Engine/Source/Programs/SubmitTool/Private/Models/ModelInterface.cpp
--- a/Engine/Source/Programs/SubmitTool/Private/Models/ModelInterface.cpp
+++ b/Engine/Source/Programs/SubmitTool/Private/Models/ModelInterface.cpp
@@ -14,7 +14,6 @@
 #include "Configuration/Configuration.h"
 #include "SourceControlOperations.h"
 #include "HAL/PlatformApplicationMisc.h"
-#include "HAL/FileManagerGeneric.h"
 #include "Async/Async.h"
 #include "Misc/StringOutputDevice.h"
 #include "Telemetry/TelemetryService.h"
@@ -704,7 +703,6 @@ void FModelInterface::StartSubmitProcess(bool bSkipShelfDialog)
 			const TArray<FString>& ShelvedFiles = ChangelistService->GetShelvedFilesDepotPaths();
 			const TArray<FString>& LocalFiles = ChangelistService->GetFilesDepotPaths();
 
-			EDialogFactoryResult DialogResult;
 			if (ShelvedFiles != LocalFiles)
 			{
 				const FText TextTitle = NSLOCTEXT("SourceControl.SubmitWindow", "ShelveConflictTitle", "Shelve - Local conflict");
@@ -735,40 +733,10 @@ void FModelInterface::StartSubmitProcess(bool bSkipShelfDialog)
 				{
 					ShelveListString = FString::Printf(TEXT("%s\n - And %d other files"), *ShelveListString, ShelvedFiles.Num() - MaxFilesToList);
 				}
-
-				if (bSkipShelfDialog)
-				{
-					DialogResult = EDialogFactoryResult::Confirm;
-				}
-				else
-				{
-					const FString Description =
-						FString::Printf(TEXT("The shelve filelist does not match the local filelist, due to p4 restrictions submit tool can only submit local content do you want to continue with the submit?\nLocal Files:\n - %s\n\nShelved Files:\n - %s"),
-							*LocalListString,
-							*ShelveListString);
-					DialogResult = FDialogFactory::ShowDialog(TextTitle, FText::FromString(Description), TArray<FString>{TEXT("Delete Shelve and Submit"), TEXT("Cancel")});
-				}
-			}
-			else
-			{
-				if (bSkipShelfDialog)
-				{
-					DialogResult = EDialogFactoryResult::Confirm;
-				}
-				else
-				{
-					const FText TextTitle = NSLOCTEXT("SourceControl.SubmitWindow", "DeleteShelvedFilesDialogTitle", "Delete shelved files?");
-					const FText TextDescription = NSLOCTEXT("SourceControl.SubmitWindow", "DeleteShelvedFilesDialogDescription", "There are shelved files in this changelist. Do you want to delete your shelf?\nIf you do not, the submit will be cancelled.");
-
-					DialogResult = FDialogFactory::ShowConfirmDialog(TextTitle, TextDescription);
-				}
 			}
 
-			if (DialogResult == EDialogFactoryResult::Confirm)
-			{
-				ChangeState(ESubmitToolAppState::Submitting, SubmitToolState == ESubmitToolAppState::SubmitLocked && bIsUserInAllowlist);
-				DeleteShelvedFiles();
-			}
+			ChangeState(ESubmitToolAppState::Submitting, SubmitToolState == ESubmitToolAppState::SubmitLocked && bIsUserInAllowlist);
+			DeleteShelvedFiles();
 		}
 		else
 		{
@@ -1238,7 +1206,8 @@ void FModelInterface::ValidateChangelist() const
 		}
 	}
 
-	PreValidate(Validators, [this](){ ValidationService->QueueAll(); });
+	PreValidate(Validators);
+	ValidationService->QueueAll();
 }
 
 void FModelInterface::ValidateSingle(const FName& ValidatorId, bool bForce) const
@@ -1253,7 +1222,8 @@ void FModelInterface::ValidateSingle(const FName& ValidatorId, bool bForce) cons
 		}
 	}
 
-	PreValidate(Validators, [this, ValidatorId, bForce]() { ValidationService->QueueSingle(ValidatorId, bForce); });
+	PreValidate(Validators);
+	ValidationService->QueueSingle(ValidatorId, bForce);
 }
 
 void FModelInterface::ValidateByArea(const ETaskArea InValidateArea) const
@@ -1267,67 +1237,27 @@ void FModelInterface::ValidateByArea(const ETaskArea InValidateArea) const
 			Validators.Add(Validator);
 		}
 	}
-	PreValidate(Validators, [this, InValidateArea]() { ValidationService->QueueByArea(InValidateArea); });
+	PreValidate(Validators);
+	ValidationService->QueueByArea(InValidateArea);
 }
 
-void FModelInterface::PreValidate(const TArray<const TWeakPtr<const FValidatorBase>>& InValidators, TFunction<void(void)> ValidationFunc) const
+void FModelInterface::PreValidate(const TArray<const TWeakPtr<const FValidatorBase>>& InValidators) const
 {
-	bool bRequiresShelve = false;
+	TRACE_CPUPROFILER_EVENT_SCOPE(FModelInterface::PreValidate);
+	bool bRequiresShelf = false;
 	for (const TWeakPtr<const FValidatorBase>& Validator : InValidators)
 	{
 		const TSharedPtr<const FValidatorBase> SharedValidator = Validator.Pin();
 		if (SharedValidator != nullptr && !SharedValidator->Definition->bIsDisabled && SharedValidator->IsRelevantToCL())
 		{
-			bRequiresShelve |= Validator.Pin()->Definition->bRequireShelveOperation;
+			bRequiresShelf |= SharedValidator->Definition->bRequireShelveOperation;
 		}
 	}
 
-	if (bRequiresShelve)
+	if (bRequiresShelf)
 	{
-		if (ChangelistService->GetShelvedFilesDepotPaths().Num() != 0)
-		{
-			FDateTime LastModifiedFile;
-			FFileManagerGeneric FileManager;
-			for (const FSourceControlStateRef& File : ChangelistService->GetFilesInCL())
-			{
-				FString Filename = File->GetFilename();
-				FFileStatData FileModifiedDate = FileManager.GetStatData(*Filename);
-				if (FileModifiedDate.bIsValid && FileModifiedDate.ModificationTime > LastModifiedFile)
-				{
-					LastModifiedFile = FileModifiedDate.ModificationTime;
-				}
-			}
-
-			TSharedPtr<ICacheDataService> CacheService = ServiceProvider->GetService<ICacheDataService>();
-			UE_LOG(LogSubmitToolDebug, Log, TEXT("Last modified file: %s, Last Shelve operation: %s"), *LastModifiedFile.ToString(), *CacheService->GetLastShelveOperationTime(ChangelistService->GetCLID()).ToString());
-			if (ChangelistService->GetFilesInCL().Num() != ChangelistService->GetShelvedFilesInCL().Num() || !CacheService.IsValid() || CacheService->GetLastShelveOperationTime(ChangelistService->GetCLID()) < LastModifiedFile)
-			{
-				const FText TextTitle = FText::FromString(TEXT("Shelf Refresh"));
-				const FText TextDescription = FText::FromString(TEXT("Some validations require a shelf, your local files do not match the shelved ones.\n\nDo you want to refresh your shelf from the latest changes in your local files?"));
-
-				EDialogFactoryResult DialogResult = FDialogFactory::ShowDialog(TextTitle, TextDescription, TArray<FString>{ TEXT("Re-shelve files"), TEXT("Use existing shelf")});
-
-				if (DialogResult == EDialogFactoryResult::FirstButton)
-				{
-					ChangelistService->DeleteShelvedFiles(FSourceControlOperationComplete::CreateLambda([ValidationFunc, this](const FSourceControlOperationRef& Operation, ECommandResult::Type Result) {
-						ChangelistService->CreateShelvedFiles(FSourceControlOperationComplete::CreateLambda([ValidationFunc](const FSourceControlOperationRef& Operation, ECommandResult::Type Result) { ValidationFunc(); }));
-						}));
-					return;
-				}
-				else if (CacheService.IsValid())
-				{
-					CacheService->SetLastShelveOperationTime(GetCLID(), FDateTime::UtcNow());
-				}
-			}			
-		}
-		else
-		{
-			ChangelistService->CreateShelvedFiles(FSourceControlOperationComplete::CreateLambda([ValidationFunc](const FSourceControlOperationRef& Operation, ECommandResult::Type Result) { ValidationFunc(); }));
-			return;
-		}
+		ChangelistService->EnsureShelfIsCurrent();
 	}
-
-	ValidationFunc();
 }
 
 void FModelInterface::ToggleValidator(const FName& ValidatorId)
diff --git a/Engine/Source/Programs/SubmitTool/Private/Models/ModelInterface.h b/Engine/Source/Programs/SubmitTool/Private/Models/ModelInterface.h
--- a/Engine/Source/Programs/SubmitTool/Private/Models/ModelInterface.h
+++ b/Engine/Source/Programs/SubmitTool/Private/Models/ModelInterface.h
@@ -265,7 +265,7 @@ private:
 
 		ValidationService->InvalidateForChanges(InChangeType);
 	}
-	void PreValidate(const TArray<const TWeakPtr<const FValidatorBase>>& InValidators, TFunction<void(void)> ValidationFunc) const;
+	void PreValidate(const TArray<const TWeakPtr<const FValidatorBase>>& InValidators) const;
 	void RefreshStateBasedOnFiles();
 	void OnChangelistReady(bool bIsValid);
 	void RevertUnchangedAndSubmit();
diff --git a/Engine/Source/Programs/SubmitTool/Private/Models/SubmitToolUserPrefs.h b/Engine/Source/Programs/SubmitTool/Private/Models/SubmitToolUserPrefs.h
--- a/Engine/Source/Programs/SubmitTool/Private/Models/SubmitToolUserPrefs.h
+++ b/Engine/Source/Programs/SubmitTool/Private/Models/SubmitToolUserPrefs.h
@@ -29,6 +29,12 @@ struct FSubmitToolUserPrefs
 
 	UPROPERTY()
 	bool bAutoUpdatePreflight = false;
+	
+	UPROPERTY()
+	bool bAutoCreateShelf = false;
+	
+	UPROPERTY()
+	bool bAutoStompShelf = false;
 
 	UPROPERTY()
 	FVector2D WindowPosition = FVector2D();
diff --git a/Engine/Source/Programs/SubmitTool/Private/SubmitToolUtils.cpp b/Engine/Source/Programs/SubmitTool/Private/SubmitToolUtils.cpp
--- a/Engine/Source/Programs/SubmitTool/Private/SubmitToolUtils.cpp
+++ b/Engine/Source/Programs/SubmitTool/Private/SubmitToolUtils.cpp
@@ -5,6 +5,9 @@
 #include "Misc/Paths.h"
 #include "Logging/SubmitToolLog.h"
 #include "Widgets/SWindow.h"
+#include "Widgets/Input/SCheckBox.h"
+#include "Widgets/Input/SButton.h"
+#include "Widgets/Text/STextBlock.h"
 #include "Framework/Application/SlateApplication.h"
 #include "HAL/FileManagerGeneric.h"
 
@@ -95,13 +98,13 @@ void FSubmitToolUtils::EnsureWindowIsInView(TSharedRef<SWindow> InWindow, bool b
 	const FPlatformRect& VirtualDisplayRect = bSingleWindow ? DisplayMetrics.PrimaryDisplayWorkAreaRect : DisplayMetrics.VirtualDisplayRect;
 
 	if(WinPos.X < VirtualDisplayRect.Left ||
-		WinPos.X <= VirtualDisplayRect.Right ||
+		WinPos.X >= VirtualDisplayRect.Right ||
 		WinPos.Y < VirtualDisplayRect.Top ||
 		WinPos.Y >= VirtualDisplayRect.Bottom)
 	{
 		FDeprecateSlateVector2D ClampedPosition;
-		ClampedPosition.X = FMath::Clamp(WinPos.X, 0.f, VirtualDisplayRect.Right - InWindow->GetSizeInScreen().X);
-		ClampedPosition.Y = FMath::Clamp(WinPos.Y, 0.f, VirtualDisplayRect.Bottom - InWindow->GetSizeInScreen().Y);
+		ClampedPosition.X = FMath::Clamp(WinPos.X, VirtualDisplayRect.Left, VirtualDisplayRect.Right - InWindow->GetSizeInScreen().X);
+		ClampedPosition.Y = FMath::Clamp(WinPos.Y, VirtualDisplayRect.Top, VirtualDisplayRect.Bottom - InWindow->GetSizeInScreen().Y);
 		InWindow->MoveWindowTo(ClampedPosition);
 	}	
 }
@@ -164,4 +167,35 @@ bool FSubmitToolUtils::IsFileInHierarchy(const FString& InWildcard, const FStrin
 	return bFound;
 }
 
+TSharedRef<SHorizontalBox> FSubmitToolUtils::BuildUserPrefCheckboxUI(bool& InUserPrefOption, const FText&& InText)
+{
+	return SNew(SHorizontalBox)
+		+ SHorizontalBox::Slot()
+		.AutoWidth()
+		.VAlign(VAlign_Center)
+		[
+			SNew(SCheckBox)
+				.IsChecked_Lambda([&InUserPrefOption]() { return InUserPrefOption ? ECheckBoxState::Checked : ECheckBoxState::Unchecked; })
+				.OnCheckStateChanged_Lambda([&InUserPrefOption](ECheckBoxState InNewState)
+					{
+						InUserPrefOption = InNewState == ECheckBoxState::Checked;
+					})
+		]
+	+ SHorizontalBox::Slot()
+		.AutoWidth()
+		.VAlign(VAlign_Center)
+		[
+			SNew(SButton)
+				.ButtonStyle(FAppStyle::Get(), "InvisibleButton")
+				.IsFocusable(false)
+				.OnClicked_Lambda([&InUserPrefOption, MoveTemp(InText)]() { InUserPrefOption = !InUserPrefOption; return FReply::Handled(); })
+				[
+					SNew(STextBlock)
+						.Justification(ETextJustify::Left)
+						.MinDesiredWidth(60)
+						.Text(InText)
+				]
+		];
+}
+
 #undef LOCTEXT_NAMESPACE
diff --git a/Engine/Source/Programs/SubmitTool/Private/SubmitToolUtils.h b/Engine/Source/Programs/SubmitTool/Private/SubmitToolUtils.h
--- a/Engine/Source/Programs/SubmitTool/Private/SubmitToolUtils.h
+++ b/Engine/Source/Programs/SubmitTool/Private/SubmitToolUtils.h
@@ -18,6 +18,8 @@ public:
 
 	static bool IsFileInHierarchy(const FString& InWildcard, const FString& InPath);
 
+	static TSharedRef<SHorizontalBox> BuildUserPrefCheckboxUI(bool& InUserPrefOption, const FText&& InText);
+
 private:
 	static TMap<FString, TMap<bool, TSet<FString>>> HierarchyWildcardsCache;
 };
diff --git a/Engine/Source/Programs/SubmitTool/Private/View/Widgets/ValidatorsWidget.cpp b/Engine/Source/Programs/SubmitTool/Private/View/Widgets/ValidatorsWidget.cpp
--- a/Engine/Source/Programs/SubmitTool/Private/View/Widgets/ValidatorsWidget.cpp
+++ b/Engine/Source/Programs/SubmitTool/Private/View/Widgets/ValidatorsWidget.cpp
@@ -32,7 +32,7 @@ void SValidatorsWidget::Construct(const FArguments& InArgs)
 	Columns = {
 		FValidatorColumn(ValidatorsColumns::Name, 1, true, [](TWeakPtr<const FValidatorBase> In) {return In.Pin()->GetValidatorName(); }),
 		FValidatorColumn(ValidatorsColumns::Required, 130, false, [](TWeakPtr<const FValidatorBase> In) {return In.Pin()->Definition->IsRequired ? TEXT("Yes") : TEXT("No"); }),
-		FValidatorColumn(ValidatorsColumns::Status, 130, false, [](TWeakPtr<const FValidatorBase> In) {return In.Pin()->GetStatusText(); }),
+		FValidatorColumn(ValidatorsColumns::Status, 150, false, [](TWeakPtr<const FValidatorBase> In) {return In.Pin()->GetStatusText(); }),
 		FValidatorColumn(ValidatorsColumns::Actions, 200, false, nullptr),
 	};
 
