SubmitTool: Add a BuildUAT validator step to compile all UAT scripts once by launching "RunUAT.bat -CompileOnly".
Run individual UAT validators in parallell by setting the environment variable "uebp_UATMutexNoWait=1".
Speed up UAT individual validators by launching them with "-NoCompileUAT".

#rnx

diff --git a/Engine/Source/Programs/SubmitTool/Private/SubmitToolApp.cpp b/Engine/Source/Programs/SubmitTool/Private/SubmitToolApp.cpp
--- a/Engine/Source/Programs/SubmitTool/Private/SubmitToolApp.cpp
+++ b/Engine/Source/Programs/SubmitTool/Private/SubmitToolApp.cpp
@@ -5,6 +5,7 @@
 #include "SubmitToolUtils.h"
 #include "Models/SubmitToolUserPrefs.h"
 
+#include "GenericPlatform/GenericPlatformMisc.h"
 #include "RequiredProgramMainCPPInclude.h"
 #include "StandaloneRenderer.h"
 #include "Framework/Application/SlateApplication.h"
@@ -201,6 +202,10 @@ int RunSubmitTool(const TCHAR* CommandLine)
 	// initialize the configuration system
 	FConfiguration::Init();
 
+	// allow running multiple parallel instances of UAT
+	FPlatformMisc::SetEnvironmentVar(TEXT("uebp_UATMutexNoWait"), TEXT("1"));
+	UE_LOG(LogSubmitToolDebug, Log, TEXT("Set environment variable uebp_UATMutexNoWait=1 to allow running multiple parallel instances of UAT."));
+
 	// App Scope
 	{
 		UE_LOG(LogSubmitTool, Log, TEXT("Version %s"), *FAppVersion::GetVersion());
