SubmitTool: Rework ISTSourceControlService/FSubmitToolPerforce::RunCommand to execute all callbacks on the GameThread.
Make SubmitTool more robust by removing the complexity of needing to synchronize ui and model data access between different threads.
Merge the double callback api FOnSCCCommandCompleteRetVal/FOnSCCCommandCompleteNoRet into a simpler FOnSCCCommandComplete api.
Rename XXXDepotStreamLength to XXXDepotStreamPath to match the perforce terminology.
Remove force fetch option for GetFilesDepotPaths(), GetShelvedFilesDepotPaths(), GetOpenFilesInChangelists(). The changelist service itself should invalidate its cached data when required so that calling the normal Get functions will always return the "correct" data, especially if validators rely on EnsureShelfIsCurrent.
Merge GetFilesDepotPaths and GetOpenFilesInChangelists into one single perforce query.
Simplify DownloadFile API.
Add a generic trace scope for all ISTSourceControlService RunCommand callbacks.

#rnx

diff --git a/Engine/Source/Programs/SubmitTool/Private/Logic/ChangelistService.cpp b/Engine/Source/Programs/SubmitTool/Private/Logic/ChangelistService.cpp
--- a/Engine/Source/Programs/SubmitTool/Private/Logic/ChangelistService.cpp
+++ b/Engine/Source/Programs/SubmitTool/Private/Logic/ChangelistService.cpp
@@ -48,9 +48,8 @@ void FChangelistService::Init()
 	TRACE_CPUPROFILER_EVENT_SCOPE(FChangelistService::Init);
 	if(SourceControlService->GetProvider().IsValid())
 	{
-		InitialP4RunTasks.Add(LaunchGetFilesDepotPaths());
+		InitialP4RunTasks.Add(LaunchGetOpenedFilesDepotPaths());
 		InitialP4RunTasks.Add(LaunchGetShelvedFilesDepotPaths());
-		InitialP4RunTasks.Add(LaunchGetOpenFilesInChangelists());
 
 		CachedSCCProvider = SourceControlService->GetProvider().Get();
 
@@ -83,7 +82,7 @@ void FChangelistService::FindInitialChangelistsAsync()
 
 	CachedSCCProvider->Execute(UpdatePendingChangelistsOperation, EConcurrency::Asynchronous, FSourceControlOperationComplete::CreateLambda([this](const FSourceControlOperationRef& UpdateOperation, ECommandResult::Type Result)
 		{
-			TRACE_CPUPROFILER_EVENT_SCOPE(FChangelistService::FindInitialChangelistsAsync_Result);
+			TRACE_CPUPROFILER_EVENT_SCOPE(FChangelistService::FindInitialChangelistsAsync_Callback);
 			if(Result == ECommandResult::Succeeded)
 			{
 				TArray<FSourceControlChangelistRef> Changelists = CachedSCCProvider->GetChangelists(EStateCacheUsage::Use);
@@ -502,29 +501,25 @@ void FChangelistService::UpdateShelvedFiles(TArray<FString> InFiles, TFunction<v
 	InFiles.Insert(TEXT("-c"), 0);
 	InFiles.Insert(CLID, 1);
 
-	SourceControlService->RunCommand(TEXT("shelve"), InFiles, FOnSCCCommandCompleteNoRet::CreateLambda(
-		[OnComplete, this](bool bSuccess, const FSCCRecordSet& InResultValues, const FSourceControlResultInfo& InResultsInfo)
+	SourceControlService->RunCommand(TEXT("shelve"), MoveTemp(InFiles), FOnSCCCommandComplete::CreateLambda(
+		[OnComplete, this](const FSCCommand& Cmd)
 		{
-			TRACE_CPUPROFILER_EVENT_SCOPE(FChangelistService::UpdateShelvedFiles_Result);
-			UE::Tasks::Launch(UE_SOURCE_LOCATION, [OnComplete, bSuccess, this]() {
-				TRACE_CPUPROFILER_EVENT_SCOPE(FChangelistService::UpdateShelvedFiles_ResultInnerTask);
-				if (bSuccess)
-				{
-					TSharedPtr<ICacheDataService> CacheService = ServiceProvider->GetService<ICacheDataService>();
-
-					if (CacheService.IsValid())
-					{
-						CacheService->SetLastShelfKnownUpToDate(GetCLID(), FDateTime::UtcNow());
-					}
-					UE_LOG(LogSubmitToolP4, Log, TEXT("CL %s shelved files have been updated"), *GetCLID());
-					bIsShelfReady = true;
-				}
+			if (Cmd.bSuccess)
+			{
+				TSharedPtr<ICacheDataService> CacheService = ServiceProvider->GetService<ICacheDataService>();
 
-				if(OnComplete != nullptr)
+				if (CacheService.IsValid())
 				{
-					OnComplete(bSuccess);
+					CacheService->SetLastShelfKnownUpToDate(GetCLID(), FDateTime::UtcNow());
 				}
-			}, LowLevelTasks::ETaskPriority::Normal, UE::Tasks::EExtendedTaskPriority::GameThreadNormalPri);			
+				UE_LOG(LogSubmitToolP4, Log, TEXT("CL %s shelved files have been updated"), *GetCLID());
+				bIsShelfReady = true;
+			}
+
+			if (OnComplete != nullptr)
+			{
+				OnComplete(Cmd.bSuccess);
+			}
 		}));
 }
 
@@ -587,7 +582,7 @@ void FChangelistService::CreateCLFromDefaultCL()
 
 	CachedSCCProvider->Execute(UpdatePendingChangelistsOperation, EConcurrency::Asynchronous, FSourceControlOperationComplete::CreateLambda(
 		[this](const FSourceControlOperationRef& Operation, ECommandResult::Type Result) {
-			TRACE_CPUPROFILER_EVENT_SCOPE(FChangelistService::CreateCLFromDefaultCL_Result);
+			TRACE_CPUPROFILER_EVENT_SCOPE(FChangelistService::CreateCLFromDefaultCL_Callback);
 
 			ActiveP4Operations.Remove(Operation);
 
@@ -707,19 +702,26 @@ void FChangelistService::PrintFilesAndShelvedFiles() const
 	}
 }
 
-UE::Tasks::FTask FChangelistService::LaunchGetFilesDepotPaths()
+UE::Tasks::FTask FChangelistService::LaunchGetOpenedFilesDepotPaths()
 {
-	return SourceControlService->RunCommand(TEXT("opened"), { "-c", CLID }, FOnSCCCommandCompleteNoRet::CreateLambda(
-		[this](bool bSuccess, const FSCCRecordSet& InResultValues, const FSourceControlResultInfo& InResultsInfo)
+	return SourceControlService->RunCommand(TEXT("opened"), { }, FOnSCCCommandComplete::CreateLambda(
+		[this](const FSCCommand& Cmd)
 		{
-			TRACE_CPUPROFILER_EVENT_SCOPE(FChangelistService::LaunchGetFilesDepotPaths_Result);
-			if(bSuccess)
+			FilesDepotPaths.Reset();
+			OpenFilesInChangelists.Reset();
+			if (Cmd.bSuccess)
 			{
-				for(const TMap<FString, FString>& Record : InResultValues)
+				for(const TMap<FString, FString>& Record : Cmd.Values)
 				{
-					if(Record.Contains(TEXT("depotFile")))
+					const FString* DepotFile = Record.Find(TEXT("depotFile"));
+					const FString* Change = Record.Find(TEXT("change"));
+					if (DepotFile && Change)
 					{
-						FilesDepotPaths.Add(Record[TEXT("depotFile")]);
+						if (*Change == CLID)
+						{
+							FilesDepotPaths.Add(*DepotFile);
+						}
+						OpenFilesInChangelists.Add(*DepotFile, *Change);
 					}
 				}
 			}
@@ -728,13 +730,12 @@ UE::Tasks::FTask FChangelistService::LaunchGetFilesDepotPaths()
 
 UE::Tasks::FTask FChangelistService::LaunchGetShelvedFilesDepotPaths()
 {
-	return SourceControlService->RunCommand(TEXT("describe"), { "-S", CLID }, FOnSCCCommandCompleteNoRet::CreateLambda(
-	[this](bool bSuccess, const FSCCRecordSet& InResultValues, const FSourceControlResultInfo& InResultsInfo)
+	return SourceControlService->RunCommand(TEXT("describe"), { "-S", CLID }, FOnSCCCommandComplete::CreateLambda(
+	[this](const FSCCommand& Cmd)
 	{
-		TRACE_CPUPROFILER_EVENT_SCOPE(FChangelistService::LaunchGetShelvedFilesDepotPaths_Result);
-		if(bSuccess)
+		if (Cmd.bSuccess)
 		{
-			for(const TMap<FString, FString>& Record : InResultValues)
+			for (const TMap<FString, FString>& Record : Cmd.Values)
 			{
 				int32 I = 0;
 				FString PathsKey = FString::Printf(TEXT("depotFile%d"), I);
@@ -749,26 +750,6 @@ UE::Tasks::FTask FChangelistService::LaunchGetShelvedFilesDepotPaths()
 	}));
 }
 
-UE::Tasks::FTask FChangelistService::LaunchGetOpenFilesInChangelists()
-{
-	return SourceControlService->RunCommand(TEXT("opened"), {}, FOnSCCCommandCompleteNoRet::CreateLambda(
-	[this](bool bSuccess, const FSCCRecordSet& InResultValues, const FSourceControlResultInfo& InResultsInfo)
-	{
-		TRACE_CPUPROFILER_EVENT_SCOPE(FChangelistService::LaunchGetOpenFilesInChangelists_Result);
-		if(bSuccess)
-		{
-			OpenFilesInChangelists.Reserve(InResultValues.Num());
-			for(const TMap<FString, FString>& Record : InResultValues)
-			{
-				if (Record.Contains(TEXT("depotFile")) && Record.Contains(TEXT("change")))
-				{
-					OpenFilesInChangelists.Add(Record[TEXT("depotFile")], Record[TEXT("change")]);
-				}
-			}
-		}
-	}));
-}
-
 void FChangelistService::WaitForInitialP4RunTasks()
 {
 	if (InitialP4RunTasks.Num())
@@ -779,22 +760,21 @@ void FChangelistService::WaitForInitialP4RunTasks()
 	}
 }
 
-const TArray<FString>& FChangelistService::GetFilesDepotPaths(bool bForce)
+const TArray<FString>& FChangelistService::GetOpenedFilesDepotPaths()
 {
 	check(IsInGameThread());
 	WaitForInitialP4RunTasks();
 
-	if (bForce || (FilesDepotPaths.IsEmpty() && !FilesInCL.IsEmpty()))
+	if (FilesDepotPaths.IsEmpty() && !FilesInCL.IsEmpty())
 	{
-		TRACE_CPUPROFILER_EVENT_SCOPE(FChangelistService::GetFilesDepotPaths);
-		FilesDepotPaths.Reset();
-		LaunchGetFilesDepotPaths().Wait();
+		TRACE_CPUPROFILER_EVENT_SCOPE(FChangelistService::GetOpenedFilesDepotPaths);
+		LaunchGetOpenedFilesDepotPaths().Wait();
 	}
 
 	return FilesDepotPaths;
 }
 
-const TArray<FString>& FChangelistService::GetShelvedFilesDepotPaths(bool bForce)
+const TArray<FString>& FChangelistService::GetShelvedFilesDepotPaths()
 {
 	check(IsInGameThread());
 	WaitForInitialP4RunTasks();
@@ -805,7 +785,7 @@ const TArray<FString>& FChangelistService::GetShelvedFilesDepotPaths(bool bForce
 		return ShelvedFilesDepotPaths;
 	}
 
-	if (bForce || (ShelvedFilesDepotPaths.IsEmpty() && !ShelvedFilesInCL.IsEmpty()))
+	if (ShelvedFilesDepotPaths.IsEmpty() && !ShelvedFilesInCL.IsEmpty())
 	{
 		TRACE_CPUPROFILER_EVENT_SCOPE(FChangelistService::GetShelvedFilesDepotPaths);
 		ShelvedFilesDepotPaths.Reset();
@@ -815,16 +795,15 @@ const TArray<FString>& FChangelistService::GetShelvedFilesDepotPaths(bool bForce
 	return ShelvedFilesDepotPaths;
 }
 
-const TMap<FString, FString>& FChangelistService::GetOpenFilesInChangelists(bool bForce)
+const TMap<FString, FString>& FChangelistService::GetOpenFilesInChangelists()
 {
 	check(IsInGameThread());
 	WaitForInitialP4RunTasks();
 
-	if (bForce || OpenFilesInChangelists.IsEmpty())
+	if (OpenFilesInChangelists.IsEmpty())
 	{
 		TRACE_CPUPROFILER_EVENT_SCOPE(FChangelistService::GetOpenFilesInChangelists);
-		OpenFilesInChangelists.Reset();
-		LaunchGetOpenFilesInChangelists().Wait();
+		LaunchGetOpenedFilesDepotPaths().Wait();
 	}
 
 	return OpenFilesInChangelists;
@@ -837,7 +816,7 @@ void FChangelistService::CheckShelfIsUpToDate(TFunction<void(bool, TArray<FStrin
 
 	TArray<FString> FilePathsOutOfDate;
 
-	if (GetFilesDepotPaths() != GetShelvedFilesDepotPaths())
+	if (GetOpenedFilesDepotPaths() != GetShelvedFilesDepotPaths())
 	{
 		OnComplete(false, MoveTemp(FilePathsOutOfDate));
 		return;
@@ -872,7 +851,7 @@ void FChangelistService::CheckShelfIsUpToDate(TFunction<void(bool, TArray<FStrin
 
 	// Only do a real p4 diff if shelf state is unknown to submit tool 
 	UE_LOG(LogSubmitToolP4, Log, TEXT("Checking Shelf diff with local files..."));
-	TArray<FString> DiffArgs = GetFilesDepotPaths();
+	TArray<FString> DiffArgs = GetOpenedFilesDepotPaths();
 	for (FString& FilePath : DiffArgs)
 	{
 		FilePath.Appendf(TEXT("@=%s"), *CLID);
@@ -881,15 +860,15 @@ void FChangelistService::CheckShelfIsUpToDate(TFunction<void(bool, TArray<FStrin
 	DiffArgs.Insert(TEXT("-Od"), 0);
 	if (!DiffTask.IsValid() || DiffTask.IsCompleted())
 	{
-		DiffTask = SourceControlService->RunCommand(TEXT("diff"), DiffArgs, FOnSCCCommandCompleteNoRet::CreateLambda(
-			[OnComplete, this](bool bSuccess, const FSCCRecordSet& InResultValues, const FSourceControlResultInfo& InResultsInfo)
+		DiffTask = SourceControlService->RunCommand(TEXT("diff"), MoveTemp(DiffArgs), FOnSCCCommandComplete::CreateLambda(
+			[OnComplete, this](const FSCCommand& Cmd)
 			{
 				TRACE_CPUPROFILER_EVENT_SCOPE(FChangelistService::CheckShelfIsUpToDate_DiffResult);
 
 				TArray<FString> FilesWithChanges;
-				if (bSuccess)
+				if (Cmd.bSuccess)
 				{
-					for (const TMap<FString, FString>& Record : InResultValues)
+					for (const TMap<FString, FString>& Record : Cmd.Values)
 					{
 						if (Record.Contains(TEXT("clientFile")))
 						{
@@ -910,16 +889,14 @@ void FChangelistService::CheckShelfIsUpToDate(TFunction<void(bool, TArray<FStrin
 					}
 				}
 
-				UE::Tasks::Launch(UE_SOURCE_LOCATION, [OnComplete, FilesWithChanges = MoveTemp(FilesWithChanges), this]() {
-					bIsShelfReady = FilesWithChanges.Num() == 0;
-					// If we've checked that there are no diffs, register it to avoid future diffs
-					if (bIsShelfReady)
-					{
-						ServiceProvider->GetService<ICacheDataService>()->SetLastShelfKnownUpToDate(CLID, FDateTime::UtcNow());
-					}
+				bIsShelfReady = Cmd.bSuccess && FilesWithChanges.Num() == 0;
+				// If we've checked that there are no diffs, register it to avoid future diffs
+				if (bIsShelfReady)
+				{
+					ServiceProvider->GetService<ICacheDataService>()->SetLastShelfKnownUpToDate(CLID, FDateTime::UtcNow());
+				}
 
-					OnComplete(bIsShelfReady, FilesWithChanges);
-				}, LowLevelTasks::ETaskPriority::Normal, UE::Tasks::EExtendedTaskPriority::GameThreadNormalPri);
+				OnComplete(bIsShelfReady, FilesWithChanges);
 			}));
 	}
 }
diff --git a/Engine/Source/Programs/SubmitTool/Private/Logic/ChangelistService.h b/Engine/Source/Programs/SubmitTool/Private/Logic/ChangelistService.h
--- a/Engine/Source/Programs/SubmitTool/Private/Logic/ChangelistService.h
+++ b/Engine/Source/Programs/SubmitTool/Private/Logic/ChangelistService.h
@@ -119,9 +119,9 @@ public:
 
 	FOnCLDescriptionUpdated OnCLDescriptionUpdated;
 
-	const TArray<FString>& GetFilesDepotPaths(bool bForce = false);
-	const TArray<FString>& GetShelvedFilesDepotPaths(bool bForce = false);
-	const TMap<FString, FString>& GetOpenFilesInChangelists(bool bForce = false);
+	const TArray<FString>& GetOpenedFilesDepotPaths();
+	const TArray<FString>& GetShelvedFilesDepotPaths();
+	const TMap<FString, FString>& GetOpenFilesInChangelists();
 
 private:
 	mutable FCriticalSection Mutex;
@@ -147,19 +147,18 @@ private:
 	FString CLDescription;
 	TArray<FSourceControlStateRef> FilesInCL;
 	TArray<FSourceControlStateRef> ShelvedFilesInCL;
-	// Data filled by duplicated queries in GetFilesDepotPaths/GetShelvedFilesDepotPaths since the SourceControl API hides this data
+	// Data filled by duplicated queries in GetOpenedFilesDepotPaths/GetShelvedFilesDepotPaths since the SourceControl API hides this data
 	TArray<FString> FilesDepotPaths;
 	TArray<FString> ShelvedFilesDepotPaths;
-	// Map filled once by GetOpenFilesInChangelists which is more efficient than a FUpdatePendingChangelistsStatus operation for all states for all pending changelists.
+	// Map filled once by GetOpenedFilesDepotPaths which is more efficient than a FUpdatePendingChangelistsStatus operation for all states for all pending changelists.
 	TMap<FString, FString> OpenFilesInChangelists;
 	bool bIsShelfReady = false;
 	UE::Tasks::FTask DiffTask;
 	void FindInitialChangelistsAsync();
 
 	void WaitForInitialP4RunTasks();
-	UE::Tasks::FTask LaunchGetFilesDepotPaths();
+	UE::Tasks::FTask LaunchGetOpenedFilesDepotPaths();
 	UE::Tasks::FTask LaunchGetShelvedFilesDepotPaths();
-	UE::Tasks::FTask LaunchGetOpenFilesInChangelists();
 
 	void FetchChangelistCallback(const FSourceControlOperationRef& UpdateOperation, ECommandResult::Type Result);
 	void CacheChangelistState();
diff --git a/Engine/Source/Programs/SubmitTool/Private/Logic/P4LockdownService.cpp b/Engine/Source/Programs/SubmitTool/Private/Logic/P4LockdownService.cpp
--- a/Engine/Source/Programs/SubmitTool/Private/Logic/P4LockdownService.cpp
+++ b/Engine/Source/Programs/SubmitTool/Private/Logic/P4LockdownService.cpp
@@ -30,6 +30,8 @@ FP4LockdownService::FP4LockdownService(const FP4LockdownParameters& InParameters
 
 bool FP4LockdownService::ArePathsInLockdown(const TArray<FString>& InPaths, bool& bOutAllowlisted)
 {	
+	TRACE_CPUPROFILER_EVENT_SCOPE(FP4LockdownService::ArePathsInLockdown);
+
 	if(!ConfigFilesTaskEvent.IsValid())
 	{
 		UE_LOG(LogSubmitToolP4, Error, TEXT("Downloading task wasn't setup correctly, hardlock status is not known."));
@@ -155,15 +157,13 @@ FP4LockdownService::FPathLockdownResult FP4LockdownService::IsPathInLockdown(con
 
 bool FP4LockdownService::GetConfigFile(const FString& InConfigId, const FString& InDepotPath)
 {
-	TArray<FSharedBuffer> FileBuffers;
-	if(ServiceProvider.Pin()->GetService<ISTSourceControlService>()->DownloadFiles(InDepotPath, FileBuffers).GetResult().bRequestSucceed && FileBuffers.Num() > 0)
+	TRACE_CPUPROFILER_EVENT_SCOPE(FP4LockdownService::GetConfigFile);
+	FSharedBuffer FileBuffer;
+	if (ServiceProvider.Pin()->GetService<ISTSourceControlService>()->DownloadFile(InDepotPath, FileBuffer).GetResult())
 	{
-		FScopeLock Scope(&FileMutex);
-		
-		DownloadedFiles.Add(InConfigId, FileBuffers[0]);
+		DownloadedFiles.Add(InConfigId, FileBuffer);
 		return true;
 	}
-
 	return false;
 }
 
@@ -209,6 +209,8 @@ void FP4LockdownService::GetAdditionalHardlockedPaths()
 
 void FP4LockdownService::ParseAllowListData()
 {
+	TRACE_CPUPROFILER_EVENT_SCOPE(FP4LockdownService::ParseAllowListData);
+
 	for(const TPair<FString, FString>& Config : Parameters.ConfigPaths)
 	{
 		FString Filepath = GetFilePath(Config.Key);
diff --git a/Engine/Source/Programs/SubmitTool/Private/Logic/P4LockdownService.h b/Engine/Source/Programs/SubmitTool/Private/Logic/P4LockdownService.h
--- a/Engine/Source/Programs/SubmitTool/Private/Logic/P4LockdownService.h
+++ b/Engine/Source/Programs/SubmitTool/Private/Logic/P4LockdownService.h
@@ -57,7 +57,6 @@ private:
 	TArray<FAllowListData> AllowListData;
 	TArray<FOverrideData> OverrideData;
 	TArray<FString> AdditionalHardlocks;
-	FCriticalSection FileMutex;
 };
 
 Expose_TNameOf(FP4LockdownService);
diff --git a/Engine/Source/Programs/SubmitTool/Private/Logic/PreflightService.cpp b/Engine/Source/Programs/SubmitTool/Private/Logic/PreflightService.cpp
--- a/Engine/Source/Programs/SubmitTool/Private/Logic/PreflightService.cpp
+++ b/Engine/Source/Programs/SubmitTool/Private/Logic/PreflightService.cpp
@@ -507,7 +507,7 @@ void FPreflightService::FetchUnlinkedPreflight(const FString& InPreflightId, boo
 
 void FPreflightService::StartPreflight()
 {
-	const TArray<FString>& ShelvedPaths = ServiceProvider.Pin()->GetService<FChangelistService>()->GetShelvedFilesDepotPaths(true);
+	const TArray<FString>& ShelvedPaths = ServiceProvider.Pin()->GetService<FChangelistService>()->GetShelvedFilesDepotPaths();
 
 	if (ShelvedPaths.IsEmpty())
 	{
@@ -598,7 +598,7 @@ bool FPreflightService::SelectPreflightTemplate(FPreflightTemplateDefinition& Ou
 	TSharedPtr<FChangelistService> ChangelistService = ServiceProvider.Pin()->GetService<FChangelistService>();
 
 	// Get the list of files in the changelist
-	const TArray<FString>& FilesInCl = bCheckShelveInstead ? ChangelistService->GetShelvedFilesDepotPaths() : ChangelistService->GetFilesDepotPaths();
+	const TArray<FString>& FilesInCl = bCheckShelveInstead ? ChangelistService->GetShelvedFilesDepotPaths() : ChangelistService->GetOpenedFilesDepotPaths();
 
 	// Loop through each definition to see if the files are in the path then check extension
 	for(const FPreflightTemplateDefinition& Def : Definition.Definitions)
@@ -626,7 +626,7 @@ FString FPreflightService::GetAdditionalTasksString(const FPreflightTemplateDefi
 	const FString EndString = TEXT("=true");
 	TSharedPtr<FChangelistService> ChangelistService = ServiceProvider.Pin()->GetService<FChangelistService>();
 
-	const TArray<FString>& FilesInCl = bCheckShelveInstead ? ChangelistService->GetShelvedFilesDepotPaths() : ChangelistService->GetFilesDepotPaths();
+	const TArray<FString>& FilesInCl = bCheckShelveInstead ? ChangelistService->GetShelvedFilesDepotPaths() : ChangelistService->GetOpenedFilesDepotPaths();
 	for(const FPreflightAdditionalTask& AdditionalTask : InTemplate.AdditionalTasks)
 	{
 		FString RegexPat = AdditionalTask.RegexPath.Replace(TEXT("$(StreamRoot)"), *StreamName, ESearchCase::IgnoreCase);
@@ -651,7 +651,7 @@ void FPreflightService::RefreshStreamName()
 	TSharedPtr<ISTSourceControlService> SCCService = ServiceProvider.Pin()->GetService<ISTSourceControlService>();
 	StreamName = SCCService->GetRootStreamName();
 
-	const TArray<FString>& ShelvedPaths = ServiceProvider.Pin()->GetService<FChangelistService>()->GetShelvedFilesDepotPaths(true);
+	const TArray<FString>& ShelvedPaths = ServiceProvider.Pin()->GetService<FChangelistService>()->GetShelvedFilesDepotPaths();
 
 	if (ShelvedPaths.Num() != 0)
 	{
@@ -670,21 +670,24 @@ void FPreflightService::RefreshStreamName()
 		if (!CommonPath.Equals(TEXT("//")))
 		{
 			int32 NextSlash = CommonPath.Find(TEXT("/"), ESearchCase::IgnoreCase, ESearchDir::FromStart, 2);
-			const FString Depot = CommonPath.Mid(2, NextSlash - 2);
-			size_t StreamDepth = SCCService->GetDepotStreamLength(Depot);
-
-			for (size_t i = 0; i < StreamDepth; ++i)
+			if (NextSlash > 2)
 			{
-				NextSlash = CommonPath.Find(TEXT("/"), ESearchCase::IgnoreCase, ESearchDir::FromStart, NextSlash + 1);
-			}
+				const FString Depot = CommonPath.Mid(2, NextSlash - 2);
+				int32 StreamDepth = SCCService->GetDepotStreamDepth(Depot);
 
-			if (NextSlash != -1)
-			{
-				const FSCCStream* FoundStream = SCCService->GetSCCStream(CommonPath.Left(NextSlash));
+				for (int32 i = 0; i < StreamDepth; ++i)
+				{
+					NextSlash = CommonPath.Find(TEXT("/"), ESearchCase::IgnoreCase, ESearchDir::FromStart, NextSlash + 1);
+				}
 
-				if (FoundStream != nullptr)
+				if (NextSlash != -1)
 				{
-					StreamName = FoundStream->Name;
+					const FSCCStream* FoundStream = SCCService->GetSCCStream(CommonPath.Left(NextSlash));
+
+					if (FoundStream != nullptr)
+					{
+						StreamName = FoundStream->Name;
+					}
 				}
 			}
 		}
diff --git a/Engine/Source/Programs/SubmitTool/Private/Logic/Services/Interfaces/ISTSourceControlService.h b/Engine/Source/Programs/SubmitTool/Private/Logic/Services/Interfaces/ISTSourceControlService.h
--- a/Engine/Source/Programs/SubmitTool/Private/Logic/Services/Interfaces/ISTSourceControlService.h
+++ b/Engine/Source/Programs/SubmitTool/Private/Logic/Services/Interfaces/ISTSourceControlService.h
@@ -75,24 +75,26 @@ public:
 	TArray<const FString> AdditionalImportPaths;
 };
 
-template<typename UserType>
-struct FSCCResultRetVal
+using FSCCRecordSet = TArray<TMap<FString, FString>>;
+
+struct FSCCommand
 {
-	bool bRequestSucceed;
-	TSharedPtr<UserType> UserData;
+	// in
+	FString Command;
+	TArray<FString> Args;
+	bool bCollectData = false;
+	// out
+	FSCCRecordSet Values;
+	FSourceControlResultInfo Results;
+	TArray<FSharedBuffer> Data;
+	bool bSuccess = false;
 };
 
-using FSCCResultNoRet = FSCCResultRetVal<void>;
-
-using FSCCRecordSet = TArray<TMap<FString, FString>>;
-
 DECLARE_MULTICAST_DELEGATE_OneParam(FOnUsersGet, TArray<TSharedPtr<FUserData>>& /*Users*/)
 DECLARE_MULTICAST_DELEGATE_OneParam(FOnGroupsGet, TArray<TSharedPtr<FString>>& /*Groups*/)
 DECLARE_MULTICAST_DELEGATE_TwoParams(FOnUsersAndGroupsGet, TArray<TSharedPtr<FUserData>>& /*Users*/, TArray<TSharedPtr<FString>>& /*Groups*/)
 
-template<typename Type>
-using FOnSCCCommandCompleteRetVal = TDelegate<Type(bool /*bSuccess*/, const FSCCRecordSet& /*InResultValues*/, const FSourceControlResultInfo& /*InResultsInfo*/)>;
-using FOnSCCCommandCompleteNoRet = FOnSCCCommandCompleteRetVal<void>;
+using FOnSCCCommandComplete = TDelegate<void(const FSCCommand&)>;
 
 class ISTSourceControlService : public ISubmitToolService
 {
@@ -102,7 +104,7 @@ public:
 	virtual void GetGroups(const FOnGroupsGet::FDelegate& Callback) = 0;
 	virtual void GetUsersAndGroups(const FOnUsersAndGroupsGet::FDelegate& Callback) = 0;
 
-	virtual UE::Tasks::TTask<FSCCResultNoRet> DownloadFiles(const FString& InFilepath, TArray<FSharedBuffer>& OutFileBuffers) = 0;
+	virtual UE::Tasks::TTask<bool> DownloadFile(const FString& InFilepath, FSharedBuffer& OutFileBuffer) = 0;
 
 	virtual bool IsAvailable() const = 0;
 
@@ -115,48 +117,38 @@ public:
 
 	virtual TSharedPtr<FUserData> GetUserDataFromCache(const FString& Username) const = 0;
 
-	virtual TArray<FSCCStream*> GetClientStreams() const = 0;
+	virtual const TArray<FSCCStream*>& GetClientStreams() const = 0;
 	virtual const FSCCStream* GetSCCStream(const FString& InStreamName) = 0;
 	virtual const FString GetRootStreamName() = 0;
 	virtual const FString GetCurrentStreamName() = 0;
-	virtual const size_t GetDepotStreamLength(const FString& InDepotName) = 0;
+	virtual int32 GetDepotStreamDepth(const FString& InDepotName) = 0;
 	virtual const FAuthTicket& GetAuthTicket() = 0;
 	virtual bool IsBusy() const = 0;
 
-	template<typename T>
-	UE::Tasks::TTask<FSCCResultRetVal<T>> RunCommand(const FString& InCommand, const TArray<FString>& InAdditionalArgs, FOnSCCCommandCompleteRetVal<T> InCompleteCallback, TArray<FSharedBuffer>* OutData = nullptr)
+	UE::Tasks::TTask<bool> RunCommand(FString&& InCommand, TArray<FString>&& InAdditionalArgs, FOnSCCCommandComplete InCompleteCallback)
 	{
-		return UE::Tasks::Launch(UE_SOURCE_LOCATION, [this, InCommand, InAdditionalArgs, InCompleteCallback, OutData] {
-			FSCCRecordSet ResultValues;
-			FSourceControlResultInfo ResultInfo;
-
-			bool bResult = RunCommandInternal(InCommand, InAdditionalArgs, ResultValues, ResultInfo, OutData);
-
-			TSharedPtr<T> UserTypeResult;
+		TSharedRef<FSCCommand> Command = MakeShared<FSCCommand>(MoveTemp(InCommand), MoveTemp(InAdditionalArgs));
+		UE::Tasks::TTask<bool> CommandTask = UE::Tasks::Launch(UE_SOURCE_LOCATION, [this, Command]
+		{
+			Command->bSuccess = RunCommandInternal(*Command);
+			return Command->bSuccess;
+		});
 
-			if (InCompleteCallback.IsBound())
+		// launch a dependent task to run the callback on the GT
+		if (InCompleteCallback.IsBound())
+		{
+			return UE::Tasks::Launch(UE_SOURCE_LOCATION, [InCompleteCallback, Command]
 			{
-				if constexpr (!std::is_same<T, void>::value)
-				{
-					UserTypeResult = MakeShared<T>(InCompleteCallback.Execute(bResult, ResultValues, ResultInfo));
-				}
-				else
-				{
-					InCompleteCallback.Execute(bResult, ResultValues, ResultInfo);
-				}
-			}
-
-			return FSCCResultRetVal<T>(bResult, MoveTemp(UserTypeResult));
-			});
-	};
-
-	UE::Tasks::TTask<FSCCResultNoRet> RunCommand(const FString& InCommand, const TArray<FString>& InAdditionalArgs, TArray<FSharedBuffer>* OutData = nullptr)
-	{
-		return RunCommand<void>(InCommand, InAdditionalArgs, nullptr, OutData);
+				TRACE_CPUPROFILER_EVENT_SCOPE_TEXT(*FString::Printf(TEXT("RunP4Command_%s_Callback"), *Command->Command));
+				InCompleteCallback.Execute(*Command);
+				return Command->bSuccess;
+			}, CommandTask, LowLevelTasks::ETaskPriority::Normal, UE::Tasks::EExtendedTaskPriority::GameThreadNormalPri);
+		}
+		return CommandTask;
 	}
 
 protected:
-	virtual bool RunCommandInternal(const FString& InCommand, const TArray<FString>& InAdditionalArgs, FSCCRecordSet& OutResultValues, FSourceControlResultInfo& OutResultInfo, TArray<FSharedBuffer>* OutData = nullptr) = 0;
+	virtual bool RunCommandInternal(FSCCommand& InCommand) = 0;
 
 };
 
diff --git a/Engine/Source/Programs/SubmitTool/Private/Logic/Services/SourceControl/SubmitToolPerforce.cpp b/Engine/Source/Programs/SubmitTool/Private/Logic/Services/SourceControl/SubmitToolPerforce.cpp
--- a/Engine/Source/Programs/SubmitTool/Private/Logic/Services/SourceControl/SubmitToolPerforce.cpp
+++ b/Engine/Source/Programs/SubmitTool/Private/Logic/Services/SourceControl/SubmitToolPerforce.cpp
@@ -175,7 +175,7 @@ const TUniquePtr<ISourceControlProvider>& FSubmitToolPerforce::GetProvider() con
 	return SCCProvider;
 }
 
-UE::Tasks::TTask<FSCCResultNoRet> FSubmitToolPerforce::GetStream(const FString& InStream, bool bRequestHierarchy)
+UE::Tasks::FTask FSubmitToolPerforce::GetStream(const FString& InStream, bool bRequestHierarchy)
 {
 	TArray<FString> Args = { TEXT("-o") };
 
@@ -184,19 +184,19 @@ UE::Tasks::TTask<FSCCResultNoRet> FSubmitToolPerforce::GetStream(const FString&
 		Args.Add(InStream);
 	}
 
-	return RunCommand(TEXT("stream"), Args, FOnSCCCommandCompleteNoRet::CreateLambda(
-		[this, bRequestHierarchy, InStream](bool bSuccess, const FSCCRecordSet& InResultValues, const FSourceControlResultInfo& InResultsInfo)
+	return RunCommand(TEXT("stream"), MoveTemp(Args), FOnSCCCommandComplete::CreateLambda(
+		[this, bRequestHierarchy, InStream](const FSCCommand& Cmd)
 		{
-			if(bSuccess && InResultValues.Num() > 0 && InResultValues[0].Contains(TEXT("Stream")))
+			if (Cmd.bSuccess && Cmd.Values.Num() > 0 && Cmd.Values[0].Contains(TEXT("Stream")))
 			{
-				TUniquePtr<FSCCStream> Stream = MakeUnique<FSCCStream>(InResultValues[0][TEXT("Stream")], InResultValues[0][TEXT("Parent")], InResultValues[0][TEXT("Type")]);
+				TUniquePtr<FSCCStream> Stream = MakeUnique<FSCCStream>(Cmd.Values[0][TEXT("Stream")], Cmd.Values[0][TEXT("Parent")], Cmd.Values[0][TEXT("Type")]);
 
 				const FString Base = TEXT("Paths");
 				size_t i = 1;
 				FString PathsKey = FString::Printf(TEXT("%s%d"), *Base, i);
-				while(InResultValues[0].Contains(PathsKey))
+				while (Cmd.Values[0].Contains(PathsKey))
 				{
-					FString Value = InResultValues[0][PathsKey];
+					FString Value = Cmd.Values[0][PathsKey];
 					if(Value.StartsWith(TEXT("import")))
 					{
 						Value = Value.RightChop(Value.Find(TEXT("//")));
@@ -214,7 +214,6 @@ UE::Tasks::TTask<FSCCResultNoRet> FSubmitToolPerforce::GetStream(const FString&
 
 				const FString Parent = Stream->Parent;
 
-				FScopeLock StreamLock(&StreamMutex);
 				if (!Streams.Contains(Stream->Name))
 				{
 					if (bRequestHierarchy)
@@ -234,7 +233,6 @@ UE::Tasks::TTask<FSCCResultNoRet> FSubmitToolPerforce::GetStream(const FString&
 			{
 				TUniquePtr<FSCCStream> Stream = MakeUnique<FSCCStream>(TEXT("Invalid"), FString(), TEXT("Invalid"));
 
-				FScopeLock StreamLock(&StreamMutex);
 				if (!InStream.IsEmpty())
 				{
 					if (!Streams.Contains(InStream))
@@ -257,22 +255,22 @@ UE::Tasks::TTask<FSCCResultNoRet> FSubmitToolPerforce::GetStream(const FString&
 
 void FSubmitToolPerforce::GetUsers(const FOnUsersGet::FDelegate& Callback)
 {
-	if(UserTask.IsValid() && UserTask.IsCompleted() && !CachedUsersArray.IsEmpty())
+	TRACE_CPUPROFILER_EVENT_SCOPE(FSubmitToolPerforce::GetUsers);
+
+	if (CachedUsersArray.Num())
 	{
 		Callback.ExecuteIfBound(CachedUsersArray);
 		return;
 	}
 
-	FScopeLock Lock(&Mutex);
 	OnUsersGetCallbacks.Add(Callback);
-	Lock.Unlock();
 
 	if(!UserTask.IsValid())
 	{
-		UserTask = RunCommand(TEXT("users"), {}, FOnSCCCommandCompleteNoRet::CreateLambda(
-			[this](bool bSuccess, const FSCCRecordSet& InResultValues, const FSourceControlResultInfo& InResultsInfo)
+		UserTask = RunCommand(TEXT("users"), {}, FOnSCCCommandComplete::CreateLambda(
+			[this](const FSCCommand& Cmd)
 			{
-				for(const TMap<FString, FString>& Record : InResultValues)
+				for(const TMap<FString, FString>& Record : Cmd.Values)
 				{
 					if(Record.Contains("User"))
 					{
@@ -284,10 +282,8 @@ void FSubmitToolPerforce::GetUsers(const FOnUsersGet::FDelegate& Callback)
 
 				LoadRecentUsers();
 
-				FScopeLock Lock(&Mutex);
 				OnUsersGetCallbacks.Broadcast(CachedUsersArray);
 				OnUsersGetCallbacks.Clear();
-				Lock.Unlock();
 			}
 		));
 	}
@@ -450,15 +446,15 @@ const FString FSubmitToolPerforce::GetRecentUsersFilepath() const
 
 void FSubmitToolPerforce::GetGroups(const FOnGroupsGet::FDelegate& Callback)
 {
-	if(GroupTask.IsValid() && GroupTask.IsCompleted() && !CachedGroupsArray.IsEmpty())
+	TRACE_CPUPROFILER_EVENT_SCOPE(FSubmitToolPerforce::GetGroups);
+
+	if (!CachedGroupsArray.IsEmpty())
 	{
 		Callback.ExecuteIfBound(CachedGroupsArray);
 		return;
 	}
 
-	FScopeLock Lock(&Mutex);
 	OnGroupsGetCallbacks.Add(Callback);
-	Lock.Unlock();
 
 	if(!GroupTask.IsValid())
 	{
@@ -570,7 +566,6 @@ void FSubmitToolPerforce::GetGroups(const FOnGroupsGet::FDelegate& Callback)
 
 				AsyncTask(ENamedThreads::GameThread, [this] 
 				{
-					FScopeLock Lock(&Mutex);
 					OnGroupsGetCallbacks.Broadcast(CachedGroupsArray);
 					OnGroupsGetCallbacks.Clear(); 
 				});
@@ -717,9 +712,19 @@ const FString FSubmitToolPerforce::GetRecentGroupsFilepath() const
 	return FPaths::Combine(FSubmitToolUtils::GetLocalAppDataPath(), TEXT("SubmitTool"), TEXT("recent_groups.dat"));
 }
 
-UE::Tasks::TTask<FSCCResultNoRet> FSubmitToolPerforce::DownloadFiles(const FString& InFilepath, TArray<FSharedBuffer>& OutFileBuffers)
+UE::Tasks::TTask<bool> FSubmitToolPerforce::DownloadFile(const FString& InFilepath, FSharedBuffer& OutFileBuffer)
 {
-	return RunCommand(TEXT("print"), { TEXT("-q"), InFilepath }, &OutFileBuffers);
+	return UE::Tasks::Launch(UE_SOURCE_LOCATION, [this, InFilepath, &OutFileBuffer]
+	{
+		FSCCommand Command(TEXT("print"), { TEXT("-q"), InFilepath });
+		Command.bCollectData = true;
+		if (RunP4Command(Command, DefaultFlags | EP4ClientUserFlags::CollectData) && Command.Data.Num())
+		{
+			OutFileBuffer = Command.Data[0];
+			return true;
+		}
+		return false;
+	});
 }
 
 void FSubmitToolPerforce::GetUsersAndGroups(const FOnUsersAndGroupsGet::FDelegate& Callback)
@@ -739,115 +744,106 @@ void FSubmitToolPerforce::GetUsersAndGroups(const FOnUsersAndGroupsGet::FDelegat
 	GetGroups(FOnGroupsGet::FDelegate::CreateLambda([OnUsersAndGroupsReady, this](TArray<TSharedPtr<FString>>&) { bGroupsReady = true; OnUsersAndGroupsReady(); }));
 }
 
-bool FSubmitToolPerforce::RunCommandInternal(const FString& InCommand, const TArray<FString>& InAdditionalArgs, FSCCRecordSet& OutResultValues, FSourceControlResultInfo& OutResultInfo, TArray<FSharedBuffer>* OutData)
+bool FSubmitToolPerforce::RunCommandInternal(FSCCommand& Command)
 {
-	EP4ClientUserFlags Flags = DefaultFlags;
-	Flags |= bIsUnicodeServer ? EP4ClientUserFlags::UnicodeServer : EP4ClientUserFlags::None;
-	
-	return RunP4Command(InCommand, InAdditionalArgs, OutResultValues, OutResultInfo, Flags, OutData);
+	return RunP4Command(Command);
 }
 
-TArray<FSCCStream*> FSubmitToolPerforce::GetClientStreams() const
+const TArray<FSCCStream*>& FSubmitToolPerforce::GetClientStreams() const
 {
-	FScopeLock StreamLock(&StreamMutex);
 	return StreamHierarchy;
 }
 
 const FSCCStream* FSubmitToolPerforce::GetSCCStream(const FString& InStreamName)
 {
+	TRACE_CPUPROFILER_EVENT_SCOPE(FSubmitToolPerforce::GetSCCStream);
+	for (int32 I = 0; I < 2; ++I)
 	{
-		FScopeLock StreamLock(&StreamMutex);
 		if (TUniquePtr<FSCCStream>* Stream = Streams.Find(InStreamName))
 		{
 			return Stream->Get();
 		}
-	}
-
-	if (GetStream(InStreamName, false).GetResult().bRequestSucceed)
-	{
-		FScopeLock StreamLock(&StreamMutex);
-		if (TUniquePtr<FSCCStream>* Stream = Streams.Find(InStreamName))
+		if (I == 0)
 		{
-			return Stream->Get();
+			GetStream(InStreamName, false).Wait();
 		}
 	}
-
 	return nullptr;
 }
 
 const FString FSubmitToolPerforce::GetRootStreamName()
 {
+	TRACE_CPUPROFILER_EVENT_SCOPE(FSubmitToolPerforce::GetRootStreamName);
+	for (int32 I = 0; I < 2; ++I)
 	{
-		FScopeLock StreamLock(&StreamMutex);
 		if (StreamHierarchy.Num() > 0)
 		{
 			return StreamHierarchy[0]->Name;
 		}
-	}
-
-	if (GetStream().GetResult().bRequestSucceed)
-	{
-		FScopeLock StreamLock(&StreamMutex);
-		if (StreamHierarchy.Num() > 0)
+		if (I == 0)
 		{
-			return StreamHierarchy[0]->Name;
+			GetStream().Wait();
 		}
 	}
-
 	return FString();
 }
 
 const FString FSubmitToolPerforce::GetCurrentStreamName()
 {
+	TRACE_CPUPROFILER_EVENT_SCOPE(FSubmitToolPerforce::GetCurrentStreamName);
+	for (int32 I = 0; I < 2; ++I)
 	{
-		FScopeLock StreamLock(&StreamMutex);
 		if (StreamHierarchy.Num() > 0)
 		{
 			return StreamHierarchy.Last()->Name;
 		}
-	}
-
-	if (GetStream().GetResult().bRequestSucceed)
-	{
-		FScopeLock StreamLock(&StreamMutex);
-		if (StreamHierarchy.Num() > 0)
+		if (I == 0)
 		{
-			return StreamHierarchy.Last()->Name;
+			GetStream().Wait();
 		}
 	}
-
 	return FString();
 }
 
-const size_t FSubmitToolPerforce::GetDepotStreamLength(const FString& InDepotName)
+UE::Tasks::FTask FSubmitToolPerforce::LaunchGetDepotStreamDepth(const FString& InDepotName)
 {
-	if(!DepotStreamLengths.Contains(InDepotName))
-	{
-		RunCommand(TEXT("depot"), { TEXT("-o"), InDepotName }, FOnSCCCommandCompleteNoRet::CreateLambda(
-			[this, InDepotName](bool bSuccess, const FSCCRecordSet& InResultValues, const FSourceControlResultInfo& InResultsInfo)
+	return RunCommand(TEXT("depot"), { TEXT("-o"), InDepotName }, FOnSCCCommandComplete::CreateLambda(
+		[this, InDepotName](const FSCCommand& Cmd)
+		{
+			if (Cmd.bSuccess && Cmd.Values.Num() > 0)
 			{
-				if(bSuccess)
+				if (const FString* StreamDepth = Cmd.Values[0].Find(TEXT("StreamDepth")))
 				{
-					if(InResultValues[0].Contains(TEXT("StreamDepth")))
+					int32 Depth = 0;
+					for (int32 I = 2; I < StreamDepth->Len(); ++I)
 					{
-						const FString& StreamDepth = InResultValues[0][TEXT("StreamDepth")];
-
-						size_t Depth = 0;
-						for(size_t i = 2; i < StreamDepth.Len(); ++i)
+						if ((*StreamDepth)[I] == TCHAR('/'))
 						{
-							if(StreamDepth[i] == TCHAR('/'))
-							{
-								Depth++;
-							}
+							Depth++;
 						}
-
-						DepotStreamLengths.Add(InDepotName, Depth);
 					}
+
+					DepotStreamDepths.Add(InDepotName, Depth);
 				}
-			})).Wait();
-	}
+			}
+		}));
+}
 
-	return DepotStreamLengths[InDepotName];
+int32 FSubmitToolPerforce::GetDepotStreamDepth(const FString& InDepotName)
+{
+	TRACE_CPUPROFILER_EVENT_SCOPE(FSubmitToolPerforce::GetDepotStreamDepth);
+	for (int32 I = 0; I < 2; ++I)
+	{
+		if (int32* StreamDepth = DepotStreamDepths.Find(InDepotName))
+		{
+			return *StreamDepth;
+		}
+		if (I == 0)
+		{
+			LaunchGetDepotStreamDepth(InDepotName).Wait();
+		}
+	}
+	return 0;
 }
 
 const FAuthTicket& FSubmitToolPerforce::GetAuthTicket()
@@ -868,12 +864,11 @@ bool FSubmitToolPerforce::IsBusy() const
 	return ActiveCommands != 0;
 }
 
-bool FSubmitToolPerforce::RunP4Command(const FString& InCommand, const TArray<FString>& InAdditionalArgs, TArray<TMap<FString, FString>>& OutResultValues, FSourceControlResultInfo& OutResults, EP4ClientUserFlags InCmdFlags, TArray<FSharedBuffer>* OutData)
+bool FSubmitToolPerforce::RunP4Command(FSCCommand& Command, EP4ClientUserFlags CmdFlags)
 {
-	TRACE_CPUPROFILER_EVENT_SCOPE_TEXT(*FString::Printf(TEXT("FSubmitToolPerforce::RunP4Command_%s"), *InCommand));
+	TRACE_CPUPROFILER_EVENT_SCOPE_TEXT(*FString::Printf(TEXT("FSubmitToolPerforce::RunP4Command_%s"), *Command.Command));
 
-	InCmdFlags |= bIsUnicodeServer ? EP4ClientUserFlags::UnicodeServer : EP4ClientUserFlags::None;
-	InCmdFlags |= OutData != nullptr ? EP4ClientUserFlags::CollectData : EP4ClientUserFlags::None;
+	CmdFlags |= bIsUnicodeServer ? EP4ClientUserFlags::UnicodeServer : EP4ClientUserFlags::None;
 
 	TUniquePtr<FP4Connection> Connection = Connections.GetAvailableConnection();
 
@@ -897,7 +892,7 @@ bool FSubmitToolPerforce::RunP4Command(const FString& InCommand, const TArray<FS
 	FString FullCommand = TEXT("p4 -p ");
 	FullCommand += TO_TCHAR(P4Client.GetPort().Text(), bIsUnicodeServer);
 
-	if(EnumHasAnyFlags(InCmdFlags, EP4ClientUserFlags::UseUser))
+	if(EnumHasAnyFlags(CmdFlags, EP4ClientUserFlags::UseUser))
 	{
 		FString UserName;
 		FCmdLineParameters::Get().GetValue(FSubmitToolCmdLine::P4User, UserName);
@@ -906,7 +901,7 @@ bool FSubmitToolPerforce::RunP4Command(const FString& InCommand, const TArray<FS
 		FullCommand += TEXT(" -u ") + UserName;
 	}
 
-	if(EnumHasAnyFlags(InCmdFlags, EP4ClientUserFlags::UseClient))
+	if(EnumHasAnyFlags(CmdFlags, EP4ClientUserFlags::UseClient))
 	{
 		FString Client;
 		FCmdLineParameters::Get().GetValue(FSubmitToolCmdLine::P4Client, Client);
@@ -915,57 +910,54 @@ bool FSubmitToolPerforce::RunP4Command(const FString& InCommand, const TArray<FS
 		FullCommand += TEXT(" -c ") + Client;
 	}
 
-	if(EnumHasAnyFlags(InCmdFlags, EP4ClientUserFlags::UseZTag))
+	if(EnumHasAnyFlags(CmdFlags, EP4ClientUserFlags::UseZTag))
 	{
 		P4Client.SetProtocol("tag", "");
 		FullCommand += TEXT(" -ztag");
 	}
 
 	FullCommand += TEXT(" ");
-	FullCommand += InCommand;
+	FullCommand += Command.Command;
 
-	int32 ArgC = InAdditionalArgs.Num();
+	int32 ArgC = Command.Args.Num();
 	UTF8CHAR** ArgV = new UTF8CHAR * [ArgC];
 	for(int32 Index = 0; Index < ArgC; Index++)
 	{
 		if(bIsUnicodeServer)
 		{
-			FTCHARToUTF8 UTF8String(*InAdditionalArgs[Index]);
+			FTCHARToUTF8 UTF8String(*Command.Args[Index]);
 			ArgV[Index] = new UTF8CHAR[UTF8String.Length() + 1];
 			FMemory::Memcpy(ArgV[Index], UTF8String.Get(), UTF8String.Length() + 1);
 		}
 		else
 		{
-			ArgV[Index] = new UTF8CHAR[InAdditionalArgs[Index].Len() + 1];
-			FMemory::Memcpy(ArgV[Index], TCHAR_TO_ANSI(*InAdditionalArgs[Index]), InAdditionalArgs[Index].Len() + 1);
+			ArgV[Index] = new UTF8CHAR[Command.Args[Index].Len() + 1];
+			FMemory::Memcpy(ArgV[Index], TCHAR_TO_ANSI(*Command.Args[Index]), Command.Args[Index].Len() + 1);
 		}
 
-		FullCommand += TEXT(" ") + InAdditionalArgs[Index];
+		FullCommand += TEXT(" ") + Command.Args[Index];
 	}
 
 	P4Client.SetArgv(ArgC, (char**)ArgV);
-	FSTClientUser P4User(OutResultValues, InCmdFlags, OutResults);
+	FSTClientUser P4User(Command.Values, CmdFlags, Command.Results);
 
 	UE_LOG(LogSubmitToolP4Debug, Log, TEXT("Running command: %s"), *FullCommand);
-	P4Client.Run(FROM_TCHAR(*InCommand, bIsUnicodeServer), &P4User);
+	P4Client.Run(FROM_TCHAR(*Command.Command, bIsUnicodeServer), &P4User);
 
-	for(const FText& Msg : OutResults.InfoMessages)
+	for(const FText& Msg : Command.Results.InfoMessages)
 	{
 		UE_LOG(LogSubmitToolP4Debug, Verbose, TEXT("%s: %s"), *FullCommand, *Msg.ToString());
 	}
 
-	if(OutResults.HasErrors())
+	if(Command.Results.HasErrors())
 	{
-		for(const FText& Error : OutResults.ErrorMessages)
+		for(const FText& Error : Command.Results.ErrorMessages)
 		{
 			UE_LOG(LogSubmitToolP4, Error, TEXT("%s: %s"), *FullCommand, *Error.ToString());
 		}
 	}
 
-	if(OutData != nullptr)
-	{
-		*OutData = P4User.ReleaseData();
-	}
+	Command.Data = P4User.ReleaseData();
 
 	// Free arguments
 	for(int32 Index = 0; Index < ArgC; Index++)
@@ -976,5 +968,5 @@ bool FSubmitToolPerforce::RunP4Command(const FString& InCommand, const TArray<FS
 
 
 	ActiveCommands--;
-	return !OutResults.HasErrors();
+	return !Command.Results.HasErrors();
 }
diff --git a/Engine/Source/Programs/SubmitTool/Private/Logic/Services/SourceControl/SubmitToolPerforce.h b/Engine/Source/Programs/SubmitTool/Private/Logic/Services/SourceControl/SubmitToolPerforce.h
--- a/Engine/Source/Programs/SubmitTool/Private/Logic/Services/SourceControl/SubmitToolPerforce.h
+++ b/Engine/Source/Programs/SubmitTool/Private/Logic/Services/SourceControl/SubmitToolPerforce.h
@@ -82,7 +82,7 @@ public:
 	virtual void GetGroups(const FOnGroupsGet::FDelegate& Callback) override;
 	virtual void GetUsersAndGroups(const FOnUsersAndGroupsGet::FDelegate& Callback) override;
 
-	virtual UE::Tasks::TTask<FSCCResultNoRet> DownloadFiles(const FString& InFilepath, TArray<FSharedBuffer>& OutFileBuffers) override;
+	virtual UE::Tasks::TTask<bool> DownloadFile(const FString& InFilepath, FSharedBuffer& OutFileBuffer) override;
 
 	virtual bool IsAvailable() const override
 	{
@@ -98,27 +98,29 @@ public:
 
 	virtual TSharedPtr<FUserData> GetUserDataFromCache(const FString& Username) const override;
 
-	virtual TArray<FSCCStream*> GetClientStreams() const override;
+	virtual const TArray<FSCCStream*>& GetClientStreams() const override;
 	virtual const FSCCStream* GetSCCStream(const FString& InStreamName) override;
 	virtual const FString GetRootStreamName() override;
 	virtual const FString GetCurrentStreamName() override;
-	virtual const size_t GetDepotStreamLength(const FString& InDepotName) override;
+	virtual int32 GetDepotStreamDepth(const FString& InDepotName) override;
 	virtual const FAuthTicket& GetAuthTicket() override;
 	virtual bool IsBusy() const override;
 			
 protected:
-	virtual bool RunCommandInternal(const FString& InCommand, const TArray<FString>& InAdditionalArgs, FSCCRecordSet& OutResultValues, FSourceControlResultInfo& OutResultInfo, TArray<FSharedBuffer>* OutData = nullptr) override;
+	virtual bool RunCommandInternal(FSCCommand& Command) override;
 
 private:
-	bool RunP4Command(const FString& InCommand, const TArray<FString>& InAdditionalArgs, TArray<TMap<FString, FString>>& OutResultValues, FSourceControlResultInfo& OutResults, EP4ClientUserFlags InCmdFlags, TArray<FSharedBuffer>* OutData = nullptr);
+	bool RunP4Command(FSCCommand& Command, EP4ClientUserFlags CmdFlags = DefaultFlags);
+	UE::Tasks::FTask LaunchGetDepotStreamDepth(const FString& InDepotName);
 
 	const FSubmitToolParameters& Parameters;
-	UE::Tasks::TTask<FSCCResultNoRet> GetStream(const FString& InStream = FString(), bool bRequestHierarchy = true);
+	UE::Tasks::FTask GetStream(const FString& InStream = FString(), bool bRequestHierarchy = true);
 	TMap<FString, TSharedPtr<FUserData>> CachedUsers;
 	TArray<TSharedPtr<FUserData>> CachedUsersArray;
 	TArray<TSharedPtr<FString>> CachedGroupsArray;
-	UE::Tasks::TTask<FSCCResultNoRet> UserTask;
+	UE::Tasks::FTask UserTask;
 	UE::Tasks::TTask<bool> GroupTask;
+	bool bIsUnicodeServer = false;
 	bool bUsersReady = false;
 	bool bGroupsReady = false;
 	FOnUsersGet OnUsersGetCallbacks;
@@ -138,12 +140,8 @@ private:
 	const FString GetRecentGroupsFilepath() const;
 	TArray<TSharedPtr<FString>> RecentGroups;
 
-	FCriticalSection Mutex;
-	bool bIsUnicodeServer = false;
-
-	mutable FCriticalSection StreamMutex;
 	TMap<FString, TUniquePtr<FSCCStream>> Streams;
-	TMap<FString, size_t> DepotStreamLengths;
+	TMap<FString, int32> DepotStreamDepths;
 
 	TArray<FSCCStream*> StreamHierarchy;
 	FAuthTicket P4Ticket;
diff --git a/Engine/Source/Programs/SubmitTool/Private/Logic/SwarmService.cpp b/Engine/Source/Programs/SubmitTool/Private/Logic/SwarmService.cpp
--- a/Engine/Source/Programs/SubmitTool/Private/Logic/SwarmService.cpp
+++ b/Engine/Source/Programs/SubmitTool/Private/Logic/SwarmService.cpp
@@ -27,11 +27,6 @@ FSwarmService::FSwarmService(TWeakPtr<FSubmitToolServiceProvider> InServiceProvi
 
 void FSwarmService::FetchReview(const OnGetReviewComplete& OnComplete)
 {
-	if (!bCanConnect)
-	{
-		return;
-	}
-
 	FString Changelist = ServiceProvider.Pin()->GetService<FChangelistService>()->GetCLID();
 
 	if (Changelist.IsEmpty())
@@ -93,16 +88,11 @@ void FSwarmService::FetchReview(const OnGetReviewComplete& OnComplete)
 
 void FSwarmService::CreateReview(const TArray<FString>& InReviewers, const OnCreateReviewComplete& OnComplete)
 {
-	if(!bCanConnect)
-	{
-		OnComplete.ExecuteIfBound(false, FString());
-		return;
-	}
-
 	FString Changelist = ServiceProvider.Pin()->GetService<FChangelistService>()->GetCLID();
 
-	if(Changelist.IsEmpty())
+	if (Changelist.IsEmpty())
 	{
+		OnComplete.ExecuteIfBound(false, FString());
 		return;
 	}
 
@@ -226,7 +216,7 @@ void FSwarmService::CreateReview(const TArray<FString>& InReviewers, const OnCre
 
 void FSwarmService::UpdateReviewDescription(const TDelegate<void(bool)>& OnComplete, const FString& InDescription)
 {
-	if(!bCanConnect || !Review.IsValid())
+	if (!Review.IsValid())
 	{
 		UE_LOG(LogSubmitTool, Warning, TEXT("Tried to update swarm review but Swarm API is not available or there is no review for this CL"));
 		OnComplete.ExecuteIfBound(false);
@@ -296,12 +286,15 @@ void FSwarmService::UpdateReviewDescription(const TDelegate<void(bool)>& OnCompl
 void FSwarmService::LaunchGetSwarmURLTask()
 {
 	GetSwarmURLTask = ServiceProvider.Pin()->GetService<ISTSourceControlService>()->RunCommand(TEXT("property"), { "-l", "-n", "P4.Swarm.URL" },
-		FOnSCCCommandCompleteNoRet::CreateLambda(
-		[this](bool bSuccess, const FSCCRecordSet& InResultValues, const FSourceControlResultInfo& InResultsInfo)
+		FOnSCCCommandComplete::CreateLambda(
+		[this](const FSCCommand& Cmd)
 		{
-			if(bSuccess && InResultValues.Num() > 0 && InResultValues[0].Contains(TEXT("value")))
+			if (Cmd.bSuccess && Cmd.Values.Num() > 0)
 			{
-				SwarmURL = InResultValues[0][TEXT("value")];
+				if (const FString* Value = Cmd.Values[0].Find(TEXT("value")))
+				{
+					SwarmURL = *Value;
+				}
 			}
 		}));
 }
diff --git a/Engine/Source/Programs/SubmitTool/Private/Logic/SwarmService.h b/Engine/Source/Programs/SubmitTool/Private/Logic/SwarmService.h
--- a/Engine/Source/Programs/SubmitTool/Private/Logic/SwarmService.h
+++ b/Engine/Source/Programs/SubmitTool/Private/Logic/SwarmService.h
@@ -72,9 +72,8 @@ private:
 	FString SwarmURL;
 	TUniquePtr<FSwarmReview> Review = nullptr;
 
-	bool bCanConnect;
-	TArray<TSharedPtr<FUserData>>* Users;
-	TArray<TSharedPtr<FString>>* Groups;
+	TArray<TSharedPtr<FUserData>>* Users = nullptr;
+	TArray<TSharedPtr<FString>>* Groups = nullptr;
 };
 
 Expose_TNameOf(FSwarmService);
diff --git a/Engine/Source/Programs/SubmitTool/Private/Logic/Validators/CoreRedirectsValidator.cpp b/Engine/Source/Programs/SubmitTool/Private/Logic/Validators/CoreRedirectsValidator.cpp
--- a/Engine/Source/Programs/SubmitTool/Private/Logic/Validators/CoreRedirectsValidator.cpp
+++ b/Engine/Source/Programs/SubmitTool/Private/Logic/Validators/CoreRedirectsValidator.cpp
@@ -40,6 +40,8 @@ void FCoreRedirectsValidator::StartAsyncWork(const FString& CLDescription, const
 
 bool FCoreRedirectsValidator::DoWork(const UE::Tasks::FCancellationToken& InCancellationToken)
 {
+	TRACE_CPUPROFILER_EVENT_SCOPE(FCoreRedirectsValidator::DoWork);
+
 	// First check the tags to avoid doing unnecessary work if the user has already promised it's safe
 	FString DocumentationURL;
 	for (const FTag* Tag : TagStatus)
@@ -60,7 +62,7 @@ bool FCoreRedirectsValidator::DoWork(const UE::Tasks::FCancellationToken& InCanc
 		}
 	}
 
-	TArray<FSharedBuffer> FileBuffers;
+	FSharedBuffer FileBuffer;
 
 	bool bAnyFilesContainRelevantModifications = false;
 	for (const FSourceControlStateRef& File : FilesToValidate)
@@ -90,20 +92,20 @@ bool FCoreRedirectsValidator::DoWork(const UE::Tasks::FCancellationToken& InCanc
 				}
 			}
 		}
-		else if (ServiceProvider.Pin()->GetService<ISTSourceControlService>()->DownloadFiles(File->GetFilename(), FileBuffers).GetResult().bRequestSucceed)
+		else if (ServiceProvider.Pin()->GetService<ISTSourceControlService>()->DownloadFile(File->GetFilename(), FileBuffer).GetResult())
 		{
 			if (InCancellationToken.IsCanceled())
 			{
 				return false;
 			}
 
-			if (FileBuffers.Num() != 1 || FileBuffers[0].GetSize() == 0)
+			if (FileBuffer.GetSize() == 0)
 			{
-				LogFailure(FString::Printf(TEXT("[%s]: Unable to perform validation because we received bad data from version control."), *GetValidatorName()));
+				LogFailure(FString::Printf(TEXT("[%s]: Unable to perform validation because we received bad data for file ('%s') from version control."), *GetValidatorName(), *File->GetFilename()));
 				return false;
 			}
 
-			const FAnsiStringView DataView(reinterpret_cast<const char*>(FileBuffers[0].GetData()));
+			const FAnsiStringView DataView(reinterpret_cast<const char*>(FileBuffer.GetData()));
 			TSet<FAnsiStringView> DepotContents;
 			PopulateSetFromStringViewOfFile(DataView, DepotContents);
 
@@ -127,7 +129,7 @@ bool FCoreRedirectsValidator::DoWork(const UE::Tasks::FCancellationToken& InCanc
 			LogFailure(FString::Printf(TEXT("[%s]: Unable to perform validation because we could not download the unedited file ('%s') from version control."), *GetValidatorName(), *File->GetFilename()));
 		}
 
-		FileBuffers.Reset();
+		FileBuffer.Reset();
 	}
 
 	return true;
diff --git a/Engine/Source/Programs/SubmitTool/Private/Logic/Validators/TagValidator.cpp b/Engine/Source/Programs/SubmitTool/Private/Logic/Validators/TagValidator.cpp
--- a/Engine/Source/Programs/SubmitTool/Private/Logic/Validators/TagValidator.cpp
+++ b/Engine/Source/Programs/SubmitTool/Private/Logic/Validators/TagValidator.cpp
@@ -70,7 +70,7 @@ bool FTagValidator::ValidateTag(const FTag* InTag, const TArray<TSharedPtr<FUser
 {
 	bool bIsTagValid = true;
 
-	FTagValidationConfig ValidationConfig = InTag->GetCurrentValidationConfig(ServiceProvider.Pin()->GetService<FChangelistService>()->GetFilesDepotPaths());
+	FTagValidationConfig ValidationConfig = InTag->GetCurrentValidationConfig(ServiceProvider.Pin()->GetService<FChangelistService>()->GetOpenedFilesDepotPaths());
 
 	if(ValidationConfig.bIsMandatory && !InTag->IsEnabled())
 	{
diff --git a/Engine/Source/Programs/SubmitTool/Private/Logic/Validators/ValidatorBase.cpp b/Engine/Source/Programs/SubmitTool/Private/Logic/Validators/ValidatorBase.cpp
--- a/Engine/Source/Programs/SubmitTool/Private/Logic/Validators/ValidatorBase.cpp
+++ b/Engine/Source/Programs/SubmitTool/Private/Logic/Validators/ValidatorBase.cpp
@@ -430,7 +430,7 @@ bool FValidatorBase::EvaluateTagSkip()
 
 	if (Definition->SkipForbiddenStreams.Num() > 0)
 	{
-		for (const FString& FileDepotPath : ChangelistService->GetFilesDepotPaths())
+		for (const FString& FileDepotPath : ChangelistService->GetOpenedFilesDepotPaths())
 		{
 			for (const FString& Stream : Definition->SkipForbiddenStreams)
 			{
diff --git a/Engine/Source/Programs/SubmitTool/Private/Models/ModelInterface.cpp b/Engine/Source/Programs/SubmitTool/Private/Models/ModelInterface.cpp
--- a/Engine/Source/Programs/SubmitTool/Private/Models/ModelInterface.cpp
+++ b/Engine/Source/Programs/SubmitTool/Private/Models/ModelInterface.cpp
@@ -653,7 +653,7 @@ void FModelInterface::StartSubmitProcess(bool bSkipShelfDialog)
 
 		if (Parameters.IncompatibleFilesParams.IncompatibleFileGroups.Num() > 0)
 		{
-			const TArray<FString>& FilesInCL = ChangelistService->GetFilesDepotPaths();
+			const TArray<FString>& FilesInCL = ChangelistService->GetOpenedFilesDepotPaths();
 
 			for (const FIncompatibleFilesGroup& FileGroup : Parameters.IncompatibleFilesParams.IncompatibleFileGroups)
 			{
@@ -701,7 +701,7 @@ void FModelInterface::StartSubmitProcess(bool bSkipShelfDialog)
 		if (HasShelvedFiles())
 		{
 			const TArray<FString>& ShelvedFiles = ChangelistService->GetShelvedFilesDepotPaths();
-			const TArray<FString>& LocalFiles = ChangelistService->GetFilesDepotPaths();
+			const TArray<FString>& LocalFiles = ChangelistService->GetOpenedFilesDepotPaths();
 
 			if (ShelvedFiles != LocalFiles)
 			{
@@ -799,11 +799,11 @@ void FModelInterface::RefreshStateBasedOnFiles()
 	}
 	else
 	{
-		TArray<FSCCStream*> Streams = SourceControlService->GetClientStreams();
+		const TArray<FSCCStream*>& Streams = SourceControlService->GetClientStreams();
 		if (!Streams.IsEmpty())
 		{
 			FString StreamsMsg = FString::JoinBy(Streams, TEXT(" -> "), [](const FSCCStream* InStr) { return InStr->Name; });
-			for (const FString& File : ChangelistService->GetFilesDepotPaths())
+			for (const FString& File : ChangelistService->GetOpenedFilesDepotPaths())
 			{
 				bool bMappedToView = false;
 
@@ -833,7 +833,7 @@ void FModelInterface::RefreshStateBasedOnFiles()
 		}
 
 		// Copy the file depot paths on the main thread since they are not safe to reference from other threads
-		UE::Tasks::Launch(UE_SOURCE_LOCATION, [this, FilesDepotPaths = ChangelistService->GetFilesDepotPaths()] {
+		UE::Tasks::Launch(UE_SOURCE_LOCATION, [this, FilesDepotPaths = ChangelistService->GetOpenedFilesDepotPaths()] {
 			if (P4LockdownService->ArePathsInLockdown(FilesDepotPaths, bIsUserInAllowlist))
 			{
 				AsyncTask(ENamedThreads::GameThread, [] { ChangeState(ESubmitToolAppState::SubmitLocked); });
diff --git a/Engine/Source/Programs/SubmitTool/Private/Models/ModelInterface.h b/Engine/Source/Programs/SubmitTool/Private/Models/ModelInterface.h
--- a/Engine/Source/Programs/SubmitTool/Private/Models/ModelInterface.h
+++ b/Engine/Source/Programs/SubmitTool/Private/Models/ModelInterface.h
@@ -119,7 +119,7 @@ public:
 	FOnCLDescriptionUpdated& GetCLDescriptionUpdatedDelegate()						{ return ChangelistService->OnCLDescriptionUpdated; }
 	
 	const TArray<FSourceControlStateRef>& GetFilesInCL() const						{ return ChangelistService->GetFilesInCL(); }
-	const TArray<FString>& GetDepotFilesInCL() const								{ return ChangelistService->GetFilesDepotPaths(); }
+	const TArray<FString>& GetDepotFilesInCL() const								{ return ChangelistService->GetOpenedFilesDepotPaths(); }
 	bool HasShelvedFiles() const													{ return ChangelistService->HasShelvedFiles(); }
 
 	bool IsP4OperationRunning(FName OperationName = FName()) const					{ return SubmitToolState == ESubmitToolAppState::P4BlockingOperation || ChangelistService->IsP4OperationRunning(OperationName); }
@@ -308,4 +308,4 @@ private:
 	static ESubmitToolAppState SubmitToolState;
 
 	static void ChangeState(ESubmitToolAppState newState, bool bForce = false);
-};
\ No newline at end of file
+};
