SubmitTool: P4LockdownService: Add back FileMutex for updating map of downloaded files (too eagerly removed in 47361297, it is still needed).

#rnx

diff --git a/Engine/Source/Programs/SubmitTool/Private/Logic/P4LockdownService.cpp b/Engine/Source/Programs/SubmitTool/Private/Logic/P4LockdownService.cpp
--- a/Engine/Source/Programs/SubmitTool/Private/Logic/P4LockdownService.cpp
+++ b/Engine/Source/Programs/SubmitTool/Private/Logic/P4LockdownService.cpp
@@ -161,6 +161,7 @@ bool FP4LockdownService::GetConfigFile(const FString& InConfigId, const FString&
 	FSharedBuffer FileBuffer;
 	if (ServiceProvider.Pin()->GetService<ISTSourceControlService>()->DownloadFile(InDepotPath, FileBuffer).GetResult())
 	{
+		FScopeLock Scope(&FileMutex);
 		DownloadedFiles.Add(InConfigId, FileBuffer);
 		return true;
 	}
diff --git a/Engine/Source/Programs/SubmitTool/Private/Logic/P4LockdownService.h b/Engine/Source/Programs/SubmitTool/Private/Logic/P4LockdownService.h
--- a/Engine/Source/Programs/SubmitTool/Private/Logic/P4LockdownService.h
+++ b/Engine/Source/Programs/SubmitTool/Private/Logic/P4LockdownService.h
@@ -57,6 +57,7 @@ private:
 	TArray<FAllowListData> AllowListData;
 	TArray<FOverrideData> OverrideData;
 	TArray<FString> AdditionalHardlocks;
+	FCriticalSection FileMutex;
 };
 
 Expose_TNameOf(FP4LockdownService);
