Adding tooltip forcefield
Removing old code

diff --git a/Engine/Source/Programs/SubmitTool/Private/Logic/IntegrationService.cpp b/Engine/Source/Programs/SubmitTool/Private/Logic/IntegrationService.cpp
--- a/Engine/Source/Programs/SubmitTool/Private/Logic/IntegrationService.cpp
+++ b/Engine/Source/Programs/SubmitTool/Private/Logic/IntegrationService.cpp
@@ -77,34 +77,6 @@ FIntegrationService::FIntegrationService
 	}
 }
 
-bool FIntegrationService::OpenIntegrationTool() const
-{
-	FString Path = FConfiguration::SubstituteAndNormalizeFilename(Parameters.Path);
-	FString Args = FConfiguration::Substitute(Parameters.Args);
-
-	if(FPaths::FileExists(Path))
-	{
-
-		FProcHandle Handle = FPlatformProcess::CreateProc(*Path, *Args, false, false, false, nullptr, 0, nullptr, nullptr, nullptr);
-		if(Handle.IsValid())
-		{
-			UE_LOG(LogSubmitTool, Log, TEXT("[FNIntegrationService] Successfully invoked the FN Integration Tool (%s) with parameters (%s)."), *Path, *Args);
-		}
-		else
-		{
-			UE_LOG(LogSubmitTool, Warning, TEXT("[FNIntegrationService] Failed to invoke the FN Integration Tool (%s) with parameters (%s)."), *Path, *Args);
-		}
-
-		return Handle.IsValid();
-	}
-	else
-	{
-		UE_LOG(LogSubmitTool, Error, TEXT("[FNIntegrationService] Could not invoke the FN Integration tool, (%s) does not exist."), *Path);
-	}
-
-	return false;
-}
-
 void FIntegrationService::RequestIntegration(const FOnBooleanValueChanged OnComplete) const
 {
 	TSharedPtr<FSwarmService> SwarmService = ServiceProvider.Pin()->GetService<FSwarmService>();
diff --git a/Engine/Source/Programs/SubmitTool/Private/Logic/IntegrationService.h b/Engine/Source/Programs/SubmitTool/Private/Logic/IntegrationService.h
--- a/Engine/Source/Programs/SubmitTool/Private/Logic/IntegrationService.h
+++ b/Engine/Source/Programs/SubmitTool/Private/Logic/IntegrationService.h
@@ -21,7 +21,6 @@ class FIntegrationService final : public ISubmitToolService
 public:
 	FIntegrationService(const FIntegrationParameters& InParameters, TWeakPtr<FSubmitToolServiceProvider> InServiceProvider);
 
-	bool OpenIntegrationTool() const;
 	void RequestIntegration(const FOnBooleanValueChanged OnComplete) const;
 
 	const TMap<FString, TSharedPtr<FIntegrationOptionBase>>& GetIntegrationOptions() const
diff --git a/Engine/Source/Programs/SubmitTool/Private/View/Widgets/SIntegrationWidget.cpp b/Engine/Source/Programs/SubmitTool/Private/View/Widgets/SIntegrationWidget.cpp
--- a/Engine/Source/Programs/SubmitTool/Private/View/Widgets/SIntegrationWidget.cpp
+++ b/Engine/Source/Programs/SubmitTool/Private/View/Widgets/SIntegrationWidget.cpp
@@ -356,7 +356,7 @@ TSharedRef<SHorizontalBox> SIntegrationWidget::CheckboxWithLabel(TSharedPtr<FInt
 
 TSharedRef<SHorizontalBox> SIntegrationWidget::TextWithLabel(TSharedPtr<FIntegrationTextOption> InOutOption)
 {
-	return SNew(SHorizontalBox)
+	TSharedRef<SHorizontalBox> Widget = SNew(SHorizontalBox)
 			+SHorizontalBox::Slot()
 			.AutoWidth()
 			.VAlign(VAlign_Center)
@@ -380,11 +380,14 @@ TSharedRef<SHorizontalBox> SIntegrationWidget::TextWithLabel(TSharedPtr<FIntegra
 						IntegrationValueChanged(InOutOption);
 					})
 			];
+			Widget->EnableToolTipForceField(true);
+
+	return Widget;
 }
 
 TSharedRef<SVerticalBox> SIntegrationWidget::MultiTextWithLabel(TSharedPtr<FIntegrationTextOption> InOutOption)
 {
-	return SNew(SVerticalBox)
+	TSharedRef<SVerticalBox> Widget = SNew(SVerticalBox)
 			+SVerticalBox::Slot()
 			.AutoHeight()			
 			[
@@ -409,6 +412,9 @@ TSharedRef<SVerticalBox> SIntegrationWidget::MultiTextWithLabel(TSharedPtr<FInte
 						IntegrationValueChanged(InOutOption); 
 					})
 			];
+			Widget->EnableToolTipForceField(true);
+
+	return Widget;
 }
 
 TSharedRef<SHorizontalBox> SIntegrationWidget::ComboWithLabel(TSharedPtr<FIntegrationComboOption> InOutOption)
