SubmitTool: Add cpu trace scopes for P4 ClientUser callbacks.

#rnx

diff --git a/Engine/Source/Programs/SubmitTool/Private/Logic/Services/SourceControl/PerforceData.h b/Engine/Source/Programs/SubmitTool/Private/Logic/Services/SourceControl/PerforceData.h
--- a/Engine/Source/Programs/SubmitTool/Private/Logic/Services/SourceControl/PerforceData.h
+++ b/Engine/Source/Programs/SubmitTool/Private/Logic/Services/SourceControl/PerforceData.h
@@ -164,6 +164,8 @@ public:
 	/**  Called by P4API when the results from running a command are ready. */
 	virtual void OutputStat(StrDict* VarList) override
 	{
+		TRACE_CPUPROFILER_EVENT_SCOPE(FSTClientUser::OutputStat);
+
 		TMap<FString, FString> Record = {};
 		StrRef Var, Value;
 
@@ -192,6 +194,8 @@ public:
 	/** Called by P4API when it output a chunk of text data from a file (commonly via P4 Print) */
 	virtual void OutputText(const char* DataPtr, int DataLength) override
 	{
+		TRACE_CPUPROFILER_EVENT_SCOPE(FSTClientUser::OutputText);
+
 		if(IsCollectingData())
 		{
 			if(File.OnDataDownloaded(DataPtr, DataLength))
@@ -223,6 +227,8 @@ public:
 	/**  Called by P4API when it output a chunk of binary data from a file (commonly via P4 Print) */
 	void OutputBinary(const char* DataPtr, int DataLength)
 	{
+		TRACE_CPUPROFILER_EVENT_SCOPE(FSTClientUser::OutputBinary);
+
 		if(IsCollectingData())
 		{
 			// For binary files we get a zero size call once the file is completed so we wait for that
@@ -266,6 +272,8 @@ public:
 
 	virtual void Message(Error* err) override
 	{
+		TRACE_CPUPROFILER_EVENT_SCOPE(FSTClientUser::Message);
+
 		StrBuf Buffer;
 		err->Fmt(Buffer, EF_PLAIN);
 
@@ -290,6 +298,8 @@ public:
 
 	virtual void OutputInfo(char Indent, const char* InInfo) override
 	{
+		TRACE_CPUPROFILER_EVENT_SCOPE(FSTClientUser::OutputInfo);
+
 		// We don't expect this to ever be called (info messages should come
 		// via ClientUser::Message) but implemented just to be safe.
 
@@ -298,6 +308,8 @@ public:
 
 	virtual void OutputError(const char* errBuf) override
 	{
+		TRACE_CPUPROFILER_EVENT_SCOPE(FSTClientUser::OutputError);
+
 		// In general we expect errors to be passed to use via ClientUser::Message but some
 		// errors raised by the p4 cpp api can call ::HandleError or ::OutputError directly.
 		// Since the default implementation of ::HandleError calls ::OutputError we only need
@@ -334,4 +346,4 @@ private:
 	FP4File File;
 };
 
-#undef LOCTEXT_NAMESPACE
\ No newline at end of file
+#undef LOCTEXT_NAMESPACE
