New Title system for messages, added lockdown message parsing and display

diff --git a/Engine/Source/Programs/SubmitTool/Private/Logic/P4LockdownService.cpp b/Engine/Source/Programs/SubmitTool/Private/Logic/P4LockdownService.cpp
--- a/Engine/Source/Programs/SubmitTool/Private/Logic/P4LockdownService.cpp
+++ b/Engine/Source/Programs/SubmitTool/Private/Logic/P4LockdownService.cpp
@@ -12,6 +12,7 @@
 #include "Internationalization/Regex.h"
 #include "Configuration/Configuration.h"
 
+
 FP4LockdownService::FP4LockdownService(const FP4LockdownParameters& InParameters, TWeakPtr<FSubmitToolServiceProvider> InServiceProvider)
 	: ConfigFilesTaskEvent(TEXT("ConfigFilesTaskEvent"))
 	, ServiceProvider(InServiceProvider)
@@ -28,14 +29,16 @@ FP4LockdownService::FP4LockdownService(const FP4LockdownParameters& InParameters
 	ConfigFilesTaskEvent.Trigger();
 }
 
-bool FP4LockdownService::ArePathsInLockdown(const TArray<FString>& InPaths, bool& bOutAllowlisted)
+FPathLockdownResult FP4LockdownService::ArePathsInLockdown(const TArray<FString>& InPaths)
 {	
 	TRACE_CPUPROFILER_EVENT_SCOPE(FP4LockdownService::ArePathsInLockdown);
 
+	FPathLockdownResult GlobalResult{ FString(), FString(), false, true };
+
 	if(!ConfigFilesTaskEvent.IsValid())
 	{
 		UE_LOG(LogSubmitToolP4, Error, TEXT("Downloading task wasn't setup correctly, hardlock status is not known."));
-		return false;
+		return GlobalResult;
 	}
 
 	if(!ConfigFilesTaskEvent.IsCompleted())
@@ -76,28 +79,39 @@ bool FP4LockdownService::ArePathsInLockdown(const TArray<FString>& InPaths, bool
 		UE_LOG(LogSubmitToolP4, Warning, TEXT("No files to check for lockdown"));
 	}
 
-	bool bLockdown = false;
-	bOutAllowlisted = true;
 	for(const FString& Path : InPaths)
 	{
 		FPathLockdownResult Result = IsPathInLockdown(Path);
-		bLockdown |= Result.bIsLocked;
 
 		if(Result.bIsLocked)
 		{
-			bOutAllowlisted &= Result.bAllowlisted;
+			UE_LOG(LogSubmitToolP4Debug, Warning, TEXT("File %s is locked down by %s.\n%s"), *Path, *Result.GroupName, *Result.Message);
+			GlobalResult.bAllowlisted &= Result.bAllowlisted;
+			GlobalResult.bIsLocked = Result.bIsLocked;
+
+			if (!GlobalResult.Message.IsEmpty())
+			{
+				GlobalResult.Message.Append(TEXT("\n"));
+			}
+			GlobalResult.Message.Append(Result.Message);
+
+			if (!GlobalResult.GroupName.IsEmpty())
+			{
+				GlobalResult.GroupName.Append(TEXT(" "));
+			}
+			GlobalResult.GroupName.Append(Result.GroupName);
 		}
 	}
 
-	return bLockdown;
+	return GlobalResult;
 }
 
-FP4LockdownService::FPathLockdownResult FP4LockdownService::IsPathInLockdown(const FString& InPath) const
+FPathLockdownResult FP4LockdownService::IsPathInLockdown(const FString& InPath) const
 {
 	FString PerforceUserName;
 	FCmdLineParameters::Get().GetValue(FSubmitToolCmdLine::P4User, PerforceUserName);
 
-	FPathLockdownResult bOverallLockdownResult { false, true };
+	FPathLockdownResult bOverallLockdownResult { FString(), FString(), false, true};
 
 	for(const FAllowListData& data : AllowListData)
 	{
@@ -133,6 +147,8 @@ FP4LockdownService::FPathLockdownResult FP4LockdownService::IsPathInLockdown(con
 
 		if (bIsLocked)
 		{
+			bOverallLockdownResult.GroupName = data.GroupName;
+			bOverallLockdownResult.Message = data.Message;
 			bOverallLockdownResult.bAllowlisted = bIsInOverrideAllowlist || data.AllowListers.Contains(PerforceUserName);
 			bOverallLockdownResult.bIsLocked = true;
 			break;
@@ -269,6 +285,12 @@ void FP4LockdownService::ParseAllowListData()
 
 				data->GroupName = ConfigPair.Key;
 
+				const FConfigValue* Message = ConfigSection.Find(TEXT("message"));
+				if (Message != nullptr)
+				{
+					data->Message = Message->GetSavedValue();
+				}
+
 				TArray<FString> Allowlisters;
 				AllowList->GetSavedValue().ParseIntoArray(Allowlisters, TEXT(","), true);
 
diff --git a/Engine/Source/Programs/SubmitTool/Private/Logic/P4LockdownService.h b/Engine/Source/Programs/SubmitTool/Private/Logic/P4LockdownService.h
--- a/Engine/Source/Programs/SubmitTool/Private/Logic/P4LockdownService.h
+++ b/Engine/Source/Programs/SubmitTool/Private/Logic/P4LockdownService.h
@@ -8,17 +8,27 @@
 #include "Tasks/Task.h"
 #include "Services/Interfaces/ISubmitToolService.h"
 #include "Memory/SharedBuffer.h"
+#include "Hash/Fnv.h"
 
 DECLARE_DELEGATE_OneParam(FOnConfigFileRetrieved, bool /*bValid*/)
 
 class FSubmitToolServiceProvider;
 
+struct FPathLockdownResult
+{
+	FString GroupName;
+	FString Message;
+	bool bIsLocked;
+	bool bAllowlisted;
+};
+
 class FP4LockdownService : public ISubmitToolService
 {	
 public:
+
 	FP4LockdownService() = delete;
 	FP4LockdownService(const FP4LockdownParameters& InParameters, TWeakPtr<FSubmitToolServiceProvider> InServiceProvider);
-	bool ArePathsInLockdown(const TArray<FString>& InPaths, bool& bOutAllowlisted);
+	FPathLockdownResult ArePathsInLockdown(const TArray<FString>& InPaths);
 	bool IsBlockingOperationRunning() const
 	{
 		return ConfigFilesTaskEvent.IsValid() && !ConfigFilesTaskEvent.IsCompleted();
@@ -26,10 +36,24 @@ public:
 
 private:
 
+	struct FCaseInsensitiveStringKeyFuncs : DefaultKeyFuncs<FString, false>
+	{
+		static FORCEINLINE bool Matches(const FString& A, const FString& B)
+		{
+			return A.Equals(B, ESearchCase::IgnoreCase);
+		}
+		static FORCEINLINE uint32 GetKeyHash(const FString& Key)
+		{
+			const FString Lower = Key.ToLower();
+			return UE::HashStringFNV1a32(MakeStringView(*Lower));
+		}
+	};
+
 	struct FAllowListData
 	{
 		FString GroupName;
-		TSet<FString> AllowListers;
+		FString Message;
+		TSet<FString, FCaseInsensitiveStringKeyFuncs> AllowListers;
 		TArray<TPair<bool, FString>> Views;
 	};
 
@@ -38,11 +62,6 @@ private:
 		TSet<FString> Sections;
 	};
 
-	struct FPathLockdownResult
-	{
-		bool bIsLocked;
-		bool bAllowlisted;
-	};
 
 	FPathLockdownResult IsPathInLockdown(const FString& InPath) const;
 	bool GetConfigFile(const FString& InConfigId, const FString& InDepotPath);
diff --git a/Engine/Source/Programs/SubmitTool/Private/Logic/TasksService.cpp b/Engine/Source/Programs/SubmitTool/Private/Logic/TasksService.cpp
--- a/Engine/Source/Programs/SubmitTool/Private/Logic/TasksService.cpp
+++ b/Engine/Source/Programs/SubmitTool/Private/Logic/TasksService.cpp
@@ -352,6 +352,7 @@ bool FTasksService::Tick(float InDeltaTime)
 				{
 					if(!DependencyTask->GetIsRunningOrQueued())
 					{
+						UE_LOG(LogValidatorsResult, Log, TEXT("%s was waiting for dependency %s but its state is %s, %s won't run"), *Task->GetValidatorName(), *DependencyTask->GetValidatorName(), *DependencyTask->GetStatusText(), *Task->GetValidatorName());
 						UE_LOG(LogValidators, Log, TEXT("%s was waiting for dependency %s but its state is %s, %s won't run"), *Task->GetValidatorName(), *DependencyTask->GetValidatorName(), *DependencyTask->GetStatusText(), *Task->GetValidatorName());
 						Task->Invalidate(true);
 					}					
diff --git a/Engine/Source/Programs/SubmitTool/Private/Models/ModelInterface.cpp b/Engine/Source/Programs/SubmitTool/Private/Models/ModelInterface.cpp
--- a/Engine/Source/Programs/SubmitTool/Private/Models/ModelInterface.cpp
+++ b/Engine/Source/Programs/SubmitTool/Private/Models/ModelInterface.cpp
@@ -80,8 +80,8 @@ FModelInterface::FModelInterface(const FSubmitToolParameters& InParameters) :
 	ServiceProvider->RegisterService<FPreflightService>(PreflightService.ToSharedRef());
 	JiraService = MakeShared<FJiraService>(InParameters.JiraParameters, 256, ServiceProvider);
 	ServiceProvider->RegisterService<FJiraService>(JiraService.ToSharedRef());
-	FNIntegrationService = MakeShared<FIntegrationService>(InParameters.IntegrationParameters, ServiceProvider);
-	ServiceProvider->RegisterService<FIntegrationService>(FNIntegrationService.ToSharedRef());
+	IntegrationService = MakeShared<FIntegrationService>(InParameters.IntegrationParameters, ServiceProvider);
+	ServiceProvider->RegisterService<FIntegrationService>(IntegrationService.ToSharedRef());
 	UpdateService = MakeShared<FUpdateService>(InParameters.HordeParameters, InParameters.AutoUpdateParameters, ServiceProvider);
 	ServiceProvider->RegisterService<FUpdateService>(UpdateService.ToSharedRef());
 
@@ -752,7 +752,7 @@ void FModelInterface::StartSubmitProcess(bool bSkipShelfDialog)
 
 void FModelInterface::RequestIntegration() const
 {
-	FNIntegrationService->RequestIntegration(FOnBooleanValueChanged::CreateLambda([this](bool bSuccess)
+	IntegrationService->RequestIntegration(FOnBooleanValueChanged::CreateLambda([this](bool bSuccess)
 		{
 			if (bSuccess)
 			{
@@ -783,10 +783,16 @@ void FModelInterface::RefreshStateBasedOnFiles()
 			{
 				// Copy the shelved file depot paths on the main thread since they are not safe to reference from other threads
 				UE::Tasks::Launch(UE_SOURCE_LOCATION, [this, ShelvedFilesDepotPaths = ChangelistService->GetShelvedFilesDepotPaths()] {
-					if (P4LockdownService->ArePathsInLockdown(ShelvedFilesDepotPaths, bIsUserInAllowlist))
+					FPathLockdownResult LockdownStatus = P4LockdownService->ArePathsInLockdown(ShelvedFilesDepotPaths);
+					if (LockdownStatus.bIsLocked)
 					{
 						UE_LOG(LogSubmitTool, Log, TEXT("There are no local files in CL %s, Submit is disabled but you can still request an Integration with your shelved files"), *ChangelistService->GetCLID());
-						AsyncTask(ENamedThreads::GameThread, [] { ChangeState(ESubmitToolAppState::SubmitLocked); });
+						UE::Tasks::Launch(UE_SOURCE_LOCATION, [this, LockdownStatus = MoveTemp(LockdownStatus)]
+						{
+							bIsUserInAllowlist = LockdownStatus.bAllowlisted;
+							TitleMessage = LockdownStatus.Message;
+							ChangeState(ESubmitToolAppState::SubmitLocked);
+						}, LowLevelTasks::ETaskPriority::Normal, UE::Tasks::EExtendedTaskPriority::GameThreadNormalPri);
 					}
 					else
 					{
@@ -834,9 +840,15 @@ void FModelInterface::RefreshStateBasedOnFiles()
 
 		// Copy the file depot paths on the main thread since they are not safe to reference from other threads
 		UE::Tasks::Launch(UE_SOURCE_LOCATION, [this, FilesDepotPaths = ChangelistService->GetOpenedFilesDepotPaths()] {
-			if (P4LockdownService->ArePathsInLockdown(FilesDepotPaths, bIsUserInAllowlist))
+			FPathLockdownResult LockdownStatus = P4LockdownService->ArePathsInLockdown(FilesDepotPaths);
+			if (LockdownStatus.bIsLocked)
 			{
-				AsyncTask(ENamedThreads::GameThread, [] { ChangeState(ESubmitToolAppState::SubmitLocked); });
+				UE::Tasks::Launch(UE_SOURCE_LOCATION, [this, LockdownStatus = MoveTemp(LockdownStatus)]
+				{
+					bIsUserInAllowlist = LockdownStatus.bAllowlisted;
+					TitleMessage = LockdownStatus.Message;
+					ChangeState(ESubmitToolAppState::SubmitLocked);
+				}, LowLevelTasks::ETaskPriority::Normal, UE::Tasks::EExtendedTaskPriority::GameThreadNormalPri);
 			}
 			else
 			{
@@ -1294,3 +1306,16 @@ bool FModelInterface::RemoveTagFromDescription(const FString& InTag, FString& In
 
 	return false;
 }
+
+const FString& FModelInterface::GetTitleMessage() const
+{
+	if(TitleMessage.IsEmpty())
+	{
+		if(FDateTime::Now().GetHour() < Parameters.GeneralParameters.EarlySubmitHour24 || FDateTime::Now().GetHour() >= Parameters.GeneralParameters.LateSubmitHour24)
+		{
+			TitleMessage = Parameters.GeneralParameters.OutOfHoursSubmitMessage.Replace(TEXT("$(DateTime)"), *FDateTime::Now().ToFormattedString(TEXT("%H:%M")));
+		}
+	}
+
+	return TitleMessage;
+}
diff --git a/Engine/Source/Programs/SubmitTool/Private/Models/ModelInterface.h b/Engine/Source/Programs/SubmitTool/Private/Models/ModelInterface.h
--- a/Engine/Source/Programs/SubmitTool/Private/Models/ModelInterface.h
+++ b/Engine/Source/Programs/SubmitTool/Private/Models/ModelInterface.h
@@ -203,10 +203,11 @@ public:
 
 	const FSubmitToolParameters& GetParameters() const { return Parameters; }
 
+	const FString& GetTitleMessage() const;
 	bool IsIntegrationRequired() const { return GetState() == ESubmitToolAppState::SubmitLocked; };
 	void RequestIntegration() const;
-	const TMap<FString, TSharedPtr<FIntegrationOptionBase>>& GetIntegrationOptions() { return FNIntegrationService->GetIntegrationOptions(); }
-	bool ValidateIntegrationOptions(bool bSilent) const	{ return FNIntegrationService->ValidateIntegrationOptions(bSilent); }
+	const TMap<FString, TSharedPtr<FIntegrationOptionBase>>& GetIntegrationOptions() { return IntegrationService->GetIntegrationOptions(); }
+	bool ValidateIntegrationOptions(bool bSilent) const	{ return IntegrationService->ValidateIntegrationOptions(bSilent); }
 
 
 	static FOnStateChanged OnStateChanged;
@@ -297,12 +298,13 @@ private:
 	TSharedPtr<FPreflightService> PreflightService;
 	TWeakPtr<SDockTab> MainTab;
 	TSharedPtr<FTasksService> PresubmitOperationsService;
-	TSharedPtr<FIntegrationService> FNIntegrationService;
+	TSharedPtr<FIntegrationService> IntegrationService;
 	TSharedPtr<FSwarmService> SwarmService;
 	TSharedPtr<FCredentialsService> CredentialsService;
 	TSharedPtr<FUpdateService> UpdateService;
 	TSharedPtr<FSubmitToolServiceProvider> ServiceProvider;
 	bool bPreflightQueued = false;
+	mutable FString TitleMessage;
 	bool Tick(float InDeltaTime);
 
 	static ESubmitToolAppState SubmitToolState;
diff --git a/Engine/Source/Programs/SubmitTool/Private/Parameters/SubmitToolParameters.h b/Engine/Source/Programs/SubmitTool/Private/Parameters/SubmitToolParameters.h
--- a/Engine/Source/Programs/SubmitTool/Private/Parameters/SubmitToolParameters.h
+++ b/Engine/Source/Programs/SubmitTool/Private/Parameters/SubmitToolParameters.h
@@ -52,6 +52,9 @@ struct FGeneralParameters
 	UPROPERTY()
 	uint8 LateSubmitHour24 = 16;
 
+	UPROPERTY()
+	FString OutOfHoursSubmitMessage = TEXT("**** It's $(DateTime) local time, be mindful of submitting late and going away, please remain alert and available for a reasonable period of time (2-3 hours) in case your changes cause issues ****");
+
 	UPROPERTY()
 	TArray<FString> GroupsToExclude;
 	
@@ -187,11 +190,6 @@ struct FIntegrationParameters
 {
 	GENERATED_BODY()
 
-	UPROPERTY()
-	FString Path;
-	UPROPERTY()
-	FString Args;
-
 	UPROPERTY()
 	TArray<FJiraIntegrationField> Fields;
 	UPROPERTY()
diff --git a/Engine/Source/Programs/SubmitTool/Private/Version/AppVersionDefines.h b/Engine/Source/Programs/SubmitTool/Private/Version/AppVersionDefines.h
--- a/Engine/Source/Programs/SubmitTool/Private/Version/AppVersionDefines.h
+++ b/Engine/Source/Programs/SubmitTool/Private/Version/AppVersionDefines.h
@@ -21,7 +21,7 @@
 
 #define SUBMIT_TOOL_VERSION_MAJOR 6
 #define SUBMIT_TOOL_VERSION_MINOR 5
-#define SUBMIT_TOOL_VERSION_PATCH 1
+#define SUBMIT_TOOL_VERSION_PATCH 2
 
 #ifndef SUBMIT_TOOL_CHANGELIST
 #define SUBMIT_TOOL_CHANGELIST_STRING TEXT("NoCL")
