Missing file change. Using TitleMessage in the submit tool UI

diff --git a/Engine/Source/Programs/SubmitTool/Private/View/Widgets/SubmitToolWidget.cpp b/Engine/Source/Programs/SubmitTool/Private/View/Widgets/SubmitToolWidget.cpp
--- a/Engine/Source/Programs/SubmitTool/Private/View/Widgets/SubmitToolWidget.cpp
+++ b/Engine/Source/Programs/SubmitTool/Private/View/Widgets/SubmitToolWidget.cpp
@@ -211,21 +211,12 @@ void SubmitToolWidget::Construct(const FArguments& InArgs)
 				[
 					SNew(SBorder)
 					.BorderBackgroundColor(FAppStyle::GetColor("ValidatorStateFail"))
-					.Visibility_Lambda([this]() {
-					if (FDateTime::Now().GetHour() < ModelInterface->GetParameters().GeneralParameters.EarlySubmitHour24 || FDateTime::Now().GetHour() >= ModelInterface->GetParameters().GeneralParameters.LateSubmitHour24)
-					{
-						return EVisibility::All;
-					}
-					else
-					{
-						return EVisibility::Collapsed;
-					}
-						})
+					.Visibility_Lambda([this]() { return !ModelInterface->GetTitleMessage().IsEmpty() ? EVisibility::All : EVisibility::Collapsed;})
 					[
 						SNew(STextBlock)
 							.Justification(ETextJustify::Center)
 							.TextStyle(FAppStyle::Get(), "BoldTextNormalSize")
-							.Text_Lambda([]() { return FText::FromString(FString::Printf(TEXT("**** It's %s local time, be mindful of submitting late and going away, please remain alert and available for a reasonable period of time (2-3 hours) in case your changes cause issues ****"), *FDateTime::Now().ToFormattedString(TEXT("%H:%M")))); })
+							.Text_Lambda([this]() { return FText::FromString(ModelInterface->GetTitleMessage()); })
 					]
 				]
 				+SVerticalBox::Slot()
