SubmitTool: Tweak DownloadFile API to be DownloadFiles.
This makes the logic simpler for the callers, and lets FSubmitToolPerforce manage the scheduling.
Run all downloads in sequentially in one single command to be nicer to Perforce.
Splitting them in two parallel commands may be more optimal in the future.

P4LockDownService:
Add missing call to GetAdditionalHardlockedPaths() in the constructor.
Break out DownloadAllowListData and WaitForAllowListData.
Protect ArePathsInLockdown with a critical section.
In ParseAllowListData make sure to reset the arrays before parsing the cached files again.

CoreRedirectsValidator:
Download all files at once instead of one at a time.
Pass the filebuffer size to the ansi string view to prevent reading data OOB.

#rnx

diff --git a/Engine/Source/Programs/SubmitTool/Private/Logic/P4LockdownService.cpp b/Engine/Source/Programs/SubmitTool/Private/Logic/P4LockdownService.cpp
--- a/Engine/Source/Programs/SubmitTool/Private/Logic/P4LockdownService.cpp
+++ b/Engine/Source/Programs/SubmitTool/Private/Logic/P4LockdownService.cpp
@@ -11,74 +11,105 @@
 #include "HAL/FileManager.h"
 #include "Internationalization/Regex.h"
 #include "Configuration/Configuration.h"
+#include "Templates/UniquePtr.h"
 
 
 FP4LockdownService::FP4LockdownService(const FP4LockdownParameters& InParameters, TWeakPtr<FSubmitToolServiceProvider> InServiceProvider)
-	: ConfigFilesTaskEvent(TEXT("ConfigFilesTaskEvent"))
-	, ServiceProvider(InServiceProvider)
+	: ServiceProvider(InServiceProvider)
 	, Parameters(InParameters)
 {
-	using namespace UE::Tasks;
-	for (const TPair<FString, FString>& Pair : Parameters.ConfigPaths)
+	GetAdditionalHardlockedPaths();
+	DownloadAllowListData();
+}
+
+void FP4LockdownService::DownloadAllowListData()
+{
+	if (Parameters.ConfigPaths.Num())
 	{
-		ConfigFilesTaskEvent.AddPrerequisites(Launch(UE_SOURCE_LOCATION, [this, &Pair]
+		TArray<FString> DepotPaths;
+		for (const TPair<FString, FString>& Pair : Parameters.ConfigPaths)
 		{
-			return GetConfigFile(Pair.Key, Pair.Value);
-		}));
+			DepotPaths.Add(Pair.Value);
+		}
+		DownloadTask = ServiceProvider.Pin()->GetService<ISTSourceControlService>()->DownloadFiles(MoveTemp(DepotPaths), DownloadedFiles);
 	}
-	ConfigFilesTaskEvent.Trigger();
 }
 
-FPathLockdownResult FP4LockdownService::ArePathsInLockdown(const TArray<FString>& InPaths)
-{	
-	TRACE_CPUPROFILER_EVENT_SCOPE(FP4LockdownService::ArePathsInLockdown);
-
-	FPathLockdownResult GlobalResult{ FString(), FString(), false, true };
+bool FP4LockdownService::WaitForAllowListData()
+{
+	TRACE_CPUPROFILER_EVENT_SCOPE(FP4LockdownService::WaitForAllowListData);
 
-	if(!ConfigFilesTaskEvent.IsValid())
+	if (!DownloadTask.IsValid())
 	{
-		UE_LOG(LogSubmitToolP4, Error, TEXT("Downloading task wasn't setup correctly, hardlock status is not known."));
-		return GlobalResult;
+		return false;
 	}
 
-	if(!ConfigFilesTaskEvent.IsCompleted())
+	if (!DownloadTask.IsCompleted())
 	{		
 		UE_LOG(LogSubmitToolP4, Log, TEXT("Waiting for download of Stream Hardlock data..."));
-		ConfigFilesTaskEvent.Wait(FTimespan::FromSeconds(5));
-
-		if(!ConfigFilesTaskEvent.IsCompleted())
+		DownloadTask.Wait(FTimespan::FromSeconds(5));
+		if (!DownloadTask.IsCompleted())
 		{
-			UE_LOG(LogSubmitToolP4, Error, TEXT("Downloading config files from P4 timed out, hardlock status is not latest, will revert to use cache."));
+			UE_LOG(LogSubmitToolP4, Warning, TEXT("Downloading config files from P4 timed out, hardlock status is not latest, will fallback to use cached files."));
+			return false; 
 		}
 	}
 
-	if(ConfigFilesTaskEvent.IsCompleted())
+	// Write the downloaded files to disk and return true if at least one file is written
+	bool bSuccess = DownloadTask.IsCompleted() && DownloadTask.GetResult() && DownloadedFiles.Num() == Parameters.ConfigPaths.Num();
+	if (bSuccess)
 	{
-		for(const TPair<FString, FSharedBuffer>& Pair : DownloadedFiles)
+		bSuccess = false;
+		int32 I = 0;
+		for (const TPair<FString, FString>& Pair : Parameters.ConfigPaths)
 		{
+			FSharedBuffer FileBuffer = DownloadedFiles[I++];
 			const FString Path = GetFilePath(Pair.Key);
-			FArchive* File = IFileManager::Get().CreateFileWriter(*Path, FILEWRITE_EvenIfReadOnly);
-
-			if(File != nullptr)
+			if (TUniquePtr<FArchive> File{IFileManager::Get().CreateFileWriter(*Path, FILEWRITE_EvenIfReadOnly)})
 			{
-				File->Serialize(const_cast<void*>(Pair.Value.GetData()), Pair.Value.GetSize());
-				File->Close();
-				delete File;
+				File->Serialize(const_cast<void*>(FileBuffer.GetData()), FileBuffer.GetSize());
+				bSuccess = true;
 			}
 			else
 			{
-				UE_LOG(LogSubmitTool, Warning, TEXT("Couldn't create lockdown file %s"), *Path)
+				UE_LOG(LogSubmitToolP4, Warning, TEXT("Couldn't create lockdown file %s"), *Path)
 			}
 		}
 	}
+	else
+	{
+		UE_LOG(LogSubmitToolP4, Warning, TEXT("Downloading config files from P4 failed, hardlock status is not latest, will fallback to use cached files."));
+	}
+	DownloadTask = {};
+	return bSuccess;
+}
 
-	ParseAllowListData();
+FPathLockdownResult FP4LockdownService::ArePathsInLockdown(const TArray<FString>& InPaths)
+{	
+	TRACE_CPUPROFILER_EVENT_SCOPE(FP4LockdownService::ArePathsInLockdown);
+
+	FPathLockdownResult GlobalResult{ FString(), FString(), false, true };
 
-	if(InPaths.IsEmpty())
+	if (InPaths.IsEmpty())
 	{
 		UE_LOG(LogSubmitToolP4, Warning, TEXT("No files to check for lockdown"));
+		return GlobalResult;
+	}
+
+	// Only allow one thread at a time to check ArePathsInLockdown
+	FScopeLock Scope(&Mutex);
+
+	// Wait for download
+	bool bShouldParseDownloadedData = (AllowListState < EAllowListState::Downloaded) ? WaitForAllowListData() : false;
+	
+	// Parse new allow list data
+	if (AllowListState == EAllowListState::Missing || bShouldParseDownloadedData)
+	{
+		ParseAllowListData();
+		AllowListState = bShouldParseDownloadedData ? EAllowListState::Downloaded : EAllowListState::Cached;
 	}
 
+	// Finally validate the paths
 	for(const FString& Path : InPaths)
 	{
 		FPathLockdownResult Result = IsPathInLockdown(Path);
@@ -171,19 +202,6 @@ FPathLockdownResult FP4LockdownService::IsPathInLockdown(const FString& InPath)
 	return bOverallLockdownResult;
 }
 
-bool FP4LockdownService::GetConfigFile(const FString& InConfigId, const FString& InDepotPath)
-{
-	TRACE_CPUPROFILER_EVENT_SCOPE(FP4LockdownService::GetConfigFile);
-	FSharedBuffer FileBuffer;
-	if (ServiceProvider.Pin()->GetService<ISTSourceControlService>()->DownloadFile(InDepotPath, FileBuffer).GetResult())
-	{
-		FScopeLock Scope(&FileMutex);
-		DownloadedFiles.Add(InConfigId, FileBuffer);
-		return true;
-	}
-	return false;
-}
-
 void RegexEscapeInline(FString& InOutRegex)
 {
 	for (int i = InOutRegex.Len() - 1; i >= 0; --i)
@@ -214,7 +232,8 @@ void RegexEscapeInline(FString& InOutRegex)
 
 void FP4LockdownService::GetAdditionalHardlockedPaths()
 {
-	for(FString EscapedViewLine : Parameters.AdditionalHardlockedPaths)
+	AdditionalHardlocks.Reserve(Parameters.AdditionalHardlockedPaths.Num());
+	for (FString EscapedViewLine : Parameters.AdditionalHardlockedPaths)
 	{
 		RegexEscapeInline(EscapedViewLine);
 		EscapedViewLine.ReplaceInline(TEXT("\\*"), TEXT("[^/]*"));
@@ -227,6 +246,8 @@ void FP4LockdownService::GetAdditionalHardlockedPaths()
 void FP4LockdownService::ParseAllowListData()
 {
 	TRACE_CPUPROFILER_EVENT_SCOPE(FP4LockdownService::ParseAllowListData);
+	AllowListData.Reset();
+	OverrideData.Reset();
 
 	for(const TPair<FString, FString>& Config : Parameters.ConfigPaths)
 	{
diff --git a/Engine/Source/Programs/SubmitTool/Private/Logic/P4LockdownService.h b/Engine/Source/Programs/SubmitTool/Private/Logic/P4LockdownService.h
--- a/Engine/Source/Programs/SubmitTool/Private/Logic/P4LockdownService.h
+++ b/Engine/Source/Programs/SubmitTool/Private/Logic/P4LockdownService.h
@@ -31,7 +31,7 @@ public:
 	FPathLockdownResult ArePathsInLockdown(const TArray<FString>& InPaths);
 	bool IsBlockingOperationRunning() const
 	{
-		return ConfigFilesTaskEvent.IsValid() && !ConfigFilesTaskEvent.IsCompleted();
+		return DownloadTask.IsValid() && !DownloadTask.IsCompleted();
 	}
 
 private:
@@ -49,6 +49,13 @@ private:
 		}
 	};
 
+	enum EAllowListState
+	{
+		Missing,
+		Cached,
+		Downloaded
+	};
+
 	struct FAllowListData
 	{
 		FString GroupName;
@@ -64,19 +71,21 @@ private:
 
 
 	FPathLockdownResult IsPathInLockdown(const FString& InPath) const;
-	bool GetConfigFile(const FString& InConfigId, const FString& InDepotPath);
+	void DownloadAllowListData();
+	bool WaitForAllowListData();
 	void ParseAllowListData();
 	void GetAdditionalHardlockedPaths();
 	FString GetFilePath(const FString& InConfigId);
 
-	TMap<FString, FSharedBuffer> DownloadedFiles;
-	UE::Tasks::FTaskEvent ConfigFilesTaskEvent;
+	TArray<FSharedBuffer> DownloadedFiles;
+	UE::Tasks::TTask<bool> DownloadTask;
 	TWeakPtr<FSubmitToolServiceProvider> ServiceProvider;
 	const FP4LockdownParameters& Parameters;
 	TArray<FAllowListData> AllowListData;
 	TArray<FOverrideData> OverrideData;
 	TArray<FString> AdditionalHardlocks;
-	FCriticalSection FileMutex;
+	EAllowListState AllowListState = EAllowListState::Missing;
+	FCriticalSection Mutex;
 };
 
 Expose_TNameOf(FP4LockdownService);
diff --git a/Engine/Source/Programs/SubmitTool/Private/Logic/Services/Interfaces/ISTSourceControlService.h b/Engine/Source/Programs/SubmitTool/Private/Logic/Services/Interfaces/ISTSourceControlService.h
--- a/Engine/Source/Programs/SubmitTool/Private/Logic/Services/Interfaces/ISTSourceControlService.h
+++ b/Engine/Source/Programs/SubmitTool/Private/Logic/Services/Interfaces/ISTSourceControlService.h
@@ -104,7 +104,7 @@ public:
 	virtual void GetGroups(const FOnGroupsGet::FDelegate& Callback) = 0;
 	virtual void GetUsersAndGroups(const FOnUsersAndGroupsGet::FDelegate& Callback) = 0;
 
-	virtual UE::Tasks::TTask<bool> DownloadFile(const FString& InFilepath, FSharedBuffer& OutFileBuffer) = 0;
+	virtual UE::Tasks::TTask<bool> DownloadFiles(TArray<FString>&& FilePaths, TArray<FSharedBuffer>& OutFileBuffers) = 0;
 
 	virtual bool IsAvailable() const = 0;
 
diff --git a/Engine/Source/Programs/SubmitTool/Private/Logic/Services/SourceControl/SubmitToolPerforce.cpp b/Engine/Source/Programs/SubmitTool/Private/Logic/Services/SourceControl/SubmitToolPerforce.cpp
--- a/Engine/Source/Programs/SubmitTool/Private/Logic/Services/SourceControl/SubmitToolPerforce.cpp
+++ b/Engine/Source/Programs/SubmitTool/Private/Logic/Services/SourceControl/SubmitToolPerforce.cpp
@@ -712,16 +712,20 @@ const FString FSubmitToolPerforce::GetRecentGroupsFilepath() const
 	return FPaths::Combine(FSubmitToolUtils::GetLocalAppDataPath(), TEXT("SubmitTool"), TEXT("recent_groups.dat"));
 }
 
-UE::Tasks::TTask<bool> FSubmitToolPerforce::DownloadFile(const FString& InFilepath, FSharedBuffer& OutFileBuffer)
+UE::Tasks::TTask<bool> FSubmitToolPerforce::DownloadFiles(TArray<FString>&& FilePaths, TArray<FSharedBuffer>& OutFileBuffers)
 {
-	return UE::Tasks::Launch(UE_SOURCE_LOCATION, [this, InFilepath, &OutFileBuffer]
+	FilePaths.Insert(TEXT("-q"), 0);
+	TSharedRef<FSCCommand> Command = MakeShared<FSCCommand>(TEXT("print"), MoveTemp(FilePaths));
+	Command->bCollectData = true;
+	return UE::Tasks::Launch(UE_SOURCE_LOCATION, [this, Command, &OutFileBuffers]
 	{
-		FSCCommand Command(TEXT("print"), { TEXT("-q"), InFilepath });
-		Command.bCollectData = true;
-		if (RunP4Command(Command, DefaultFlags | EP4ClientUserFlags::CollectData) && Command.Data.Num())
+		if (RunP4Command(*Command, DefaultFlags | EP4ClientUserFlags::CollectData))
 		{
-			OutFileBuffer = Command.Data[0];
-			return true;
+			if (Command->Data.Num() == Command->Args.Num() - 1)
+			{
+				OutFileBuffers = MoveTemp(Command->Data);
+				return true;
+			}
 		}
 		return false;
 	});
diff --git a/Engine/Source/Programs/SubmitTool/Private/Logic/Services/SourceControl/SubmitToolPerforce.h b/Engine/Source/Programs/SubmitTool/Private/Logic/Services/SourceControl/SubmitToolPerforce.h
--- a/Engine/Source/Programs/SubmitTool/Private/Logic/Services/SourceControl/SubmitToolPerforce.h
+++ b/Engine/Source/Programs/SubmitTool/Private/Logic/Services/SourceControl/SubmitToolPerforce.h
@@ -82,7 +82,7 @@ public:
 	virtual void GetGroups(const FOnGroupsGet::FDelegate& Callback) override;
 	virtual void GetUsersAndGroups(const FOnUsersAndGroupsGet::FDelegate& Callback) override;
 
-	virtual UE::Tasks::TTask<bool> DownloadFile(const FString& InFilepath, FSharedBuffer& OutFileBuffer) override;
+	virtual UE::Tasks::TTask<bool> DownloadFiles(TArray<FString>&& FilePaths, TArray<FSharedBuffer>& OutFileBuffers) override;
 
 	virtual bool IsAvailable() const override
 	{
diff --git a/Engine/Source/Programs/SubmitTool/Private/Logic/Validators/CoreRedirectsValidator.cpp b/Engine/Source/Programs/SubmitTool/Private/Logic/Validators/CoreRedirectsValidator.cpp
--- a/Engine/Source/Programs/SubmitTool/Private/Logic/Validators/CoreRedirectsValidator.cpp
+++ b/Engine/Source/Programs/SubmitTool/Private/Logic/Validators/CoreRedirectsValidator.cpp
@@ -63,8 +63,9 @@ bool FCoreRedirectsValidator::DoWork(const UE::Tasks::FCancellationToken& InCanc
 		}
 	}
 
-	FSharedBuffer FileBuffer;
+	TArray<FString> FilePaths;
 
+	// validate added files
 	bool bAnyFilesContainRelevantModifications = false;
 	for (const FSCFileRef& File : FilesToValidate)
 	{
@@ -93,44 +94,56 @@ bool FCoreRedirectsValidator::DoWork(const UE::Tasks::FCancellationToken& InCanc
 				}
 			}
 		}
-		else if (ServiceProvider.Pin()->GetService<ISTSourceControlService>()->DownloadFile(File->GetFilename(), FileBuffer).GetResult())
+		else
 		{
-			if (InCancellationToken.IsCanceled())
-			{
-				return false;
-			}
+			FilePaths.Add(File->GetFilename());
+		}
+	}
 
-			if (FileBuffer.GetSize() == 0)
-			{
-				LogFailure(FString::Printf(TEXT("[%s]: Unable to perform validation because we received bad data for file ('%s') from version control."), *GetValidatorName(), *File->GetFilename()));
-				return false;
-			}
+	// download remaining files
+	TArray<FSharedBuffer> FileBuffers;
+	if (!ServiceProvider.Pin()->GetService<ISTSourceControlService>()->DownloadFiles(CopyTemp(FilePaths), FileBuffers).GetResult() || FileBuffers.Num() != FilePaths.Num())
+	{
+		LogFailure(FString::Printf(TEXT("[%s]: Unable to perform validation because we could not download %d unedited file(s) from version control."), *GetValidatorName(), FilePaths.Num()));
+		return false;
+	}
 
-			const FAnsiStringView DataView(reinterpret_cast<const char*>(FileBuffer.GetData()));
-			TSet<FAnsiStringView> DepotContents;
-			PopulateSetFromStringViewOfFile(DataView, DepotContents);
+	// validate downloaded files
+	TRACE_CPUPROFILER_EVENT_SCOPE(ValidateDownloadedFiles);
+	for (int32 I = 0, N = FilePaths.Num(); I < N; ++I)
+	{
+		if (InCancellationToken.IsCanceled())
+		{
+			return false;
+		}
 
-			TArray<FString> WorkspaceContents;
-			if (!FFileHelper::LoadFileToStringArray(WorkspaceContents, *File->GetFilename()))
-			{
-				LogFailure(FString::Printf(TEXT("[%s]: Unable to perform validation because we failed to load '%s'."), *GetValidatorName(), *File->GetFilename()));
-				return false;
-			}
+		const FString& Filename = FilePaths[I];
+		const FSharedBuffer& FileBuffer = FileBuffers[I];
 
-			if (ContainsModifiedRedirects(DepotContents, WorkspaceContents))
-			{
-				LogFailure(FString::Printf(TEXT("[%s]: Your changelist includes changes to PackageRedirectors via ini file. Please read the linked documentation"), *GetValidatorName()));
-				LogFailure(FString::Printf(TEXT("[%s]: %s"), *GetValidatorName(), *DocumentationURL));
-				LogFailure(FString::Printf(TEXT("[%s]: If your changes are safe, add #redirectsaresafe"), *GetValidatorName()));
-				return false;
-			}
+		if (FileBuffer.GetSize() == 0)
+		{
+			LogFailure(FString::Printf(TEXT("[%s]: Unable to perform validation because we received bad data for file ('%s') from version control."), *GetValidatorName(), *Filename));
+			return false;
 		}
-		else
+
+		const FAnsiStringView DataView(reinterpret_cast<const char*>(FileBuffer.GetData()), FileBuffer.GetSize());
+		TSet<FAnsiStringView> DepotContents;
+		PopulateSetFromStringViewOfFile(DataView, DepotContents);
+
+		TArray<FString> WorkspaceContents;
+		if (!FFileHelper::LoadFileToStringArray(WorkspaceContents, *Filename))
 		{
-			LogFailure(FString::Printf(TEXT("[%s]: Unable to perform validation because we could not download the unedited file ('%s') from version control."), *GetValidatorName(), *File->GetFilename()));
+			LogFailure(FString::Printf(TEXT("[%s]: Unable to perform validation because we failed to load '%s'."), *GetValidatorName(), *Filename));
+			return false;
 		}
 
-		FileBuffer.Reset();
+		if (ContainsModifiedRedirects(DepotContents, WorkspaceContents))
+		{
+			LogFailure(FString::Printf(TEXT("[%s]: Your changelist includes changes to PackageRedirectors via ini file. Please read the linked documentation"), *GetValidatorName()));
+			LogFailure(FString::Printf(TEXT("[%s]: %s"), *GetValidatorName(), *DocumentationURL));
+			LogFailure(FString::Printf(TEXT("[%s]: If your changes are safe, add #redirectsaresafe"), *GetValidatorName()));
+			return false;
+		}
 	}
 
 	return true;
