SubmitTool: CoreRedirectsValidator: Add explicit handling of changelist with only Engine.ini files opened for add.

#rnx

diff --git a/Engine/Source/Programs/SubmitTool/Private/Logic/Validators/CoreRedirectsValidator.cpp b/Engine/Source/Programs/SubmitTool/Private/Logic/Validators/CoreRedirectsValidator.cpp
--- a/Engine/Source/Programs/SubmitTool/Private/Logic/Validators/CoreRedirectsValidator.cpp
+++ b/Engine/Source/Programs/SubmitTool/Private/Logic/Validators/CoreRedirectsValidator.cpp
@@ -100,6 +100,11 @@ bool FCoreRedirectsValidator::DoWork(const UE::Tasks::FCancellationToken& InCanc
 		}
 	}
 
+	if (FilePaths.IsEmpty())
+	{
+		return true;
+	}
+
 	// download remaining files
 	TArray<FSharedBuffer> FileBuffers;
 	if (!ServiceProvider.Pin()->GetService<ISTSourceControlService>()->DownloadFiles(CopyTemp(FilePaths), FileBuffers).GetResult() || FileBuffers.Num() != FilePaths.Num())
