Fix running submit tool on the default changelist checking shelf

diff --git a/Engine/Source/Programs/SubmitTool/Private/Logic/ChangelistService.cpp b/Engine/Source/Programs/SubmitTool/Private/Logic/ChangelistService.cpp
--- a/Engine/Source/Programs/SubmitTool/Private/Logic/ChangelistService.cpp
+++ b/Engine/Source/Programs/SubmitTool/Private/Logic/ChangelistService.cpp
@@ -49,7 +49,6 @@ void FChangelistService::Init()
 	if(SourceControlService->GetProvider().IsValid())
 	{
 		InitialP4RunTasks.Add(LaunchGetOpenedFilesDepotPaths());
-		InitialP4RunTasks.Add(LaunchGetShelvedFilesDepotPaths());
 
 		CachedSCCProvider = SourceControlService->GetProvider().Get();
 
@@ -59,6 +58,8 @@ void FChangelistService::Init()
 		}
 		else
 		{
+			// Only do a p4 describe if we're not using the default CL
+			InitialP4RunTasks.Add(LaunchGetShelvedFilesDepotPaths());
 			FindInitialChangelistsAsync();
 		}
 		
