Adding FromEditor only parameters to allow conditional arguments for validators if submit tool was invoked from the editor (using -from-editor cmdline arg)

diff --git a/Engine/Source/Programs/SubmitTool/Private/Logic/Validators/ValidatorDefinition.h b/Engine/Source/Programs/SubmitTool/Private/Logic/Validators/ValidatorDefinition.h
--- a/Engine/Source/Programs/SubmitTool/Private/Logic/Validators/ValidatorDefinition.h
+++ b/Engine/Source/Programs/SubmitTool/Private/Logic/Validators/ValidatorDefinition.h
@@ -304,7 +304,13 @@ struct FValidatorRunExecutableDefinition : public FValidatorDefinition
 	///	Arguments to pass to the Executable
 	/// 
 	UPROPERTY()
-	FString ExecutableArguments;
+	FString ExecutableArguments;	
+		
+	///
+	///	These arguments are only added when running submit tool from the editor
+	///
+	UPROPERTY()
+	FString FromEditorOnlyArguments;
 
 	UPROPERTY()
 	FString FileInCLArgument;
diff --git a/Engine/Source/Programs/SubmitTool/Private/Logic/Validators/ValidatorRunExecutable.cpp b/Engine/Source/Programs/SubmitTool/Private/Logic/Validators/ValidatorRunExecutable.cpp
--- a/Engine/Source/Programs/SubmitTool/Private/Logic/Validators/ValidatorRunExecutable.cpp
+++ b/Engine/Source/Programs/SubmitTool/Private/Logic/Validators/ValidatorRunExecutable.cpp
@@ -85,7 +85,6 @@ bool FValidatorRunExecutable::Validate(const FString& InCLDescription, const TAr
 
 		FinalArgs += FString::Printf(TEXT(" %s\"%s\""), *TypedDefinition->FileListArgument, *FileListPath);
 	}
-
 	// Legacy, pass down each file to the validator in the command line
 	else if(!TypedDefinition->FileInCLArgument.IsEmpty())
 	{
@@ -95,6 +94,12 @@ bool FValidatorRunExecutable::Validate(const FString& InCLDescription, const TAr
 		}
 	}
 
+	if (!TypedDefinition->FromEditorOnlyArguments.IsEmpty() && FCmdLineParameters::Get().Contains(FSubmitToolCmdLine::EditorFlag))
+	{
+		FinalArgs.Append(TEXT(" "));
+		FinalArgs.Append(TypedDefinition->FromEditorOnlyArguments);
+	}
+
 	FString ExecutablePath;
 	if(TypedDefinition->ExecutableCandidates.Num() != 0)
 	{
