Fix appending repeated messages for each file in the lockdown list

diff --git a/Engine/Source/Programs/SubmitTool/Private/Logic/P4LockdownService.cpp b/Engine/Source/Programs/SubmitTool/Private/Logic/P4LockdownService.cpp
--- a/Engine/Source/Programs/SubmitTool/Private/Logic/P4LockdownService.cpp
+++ b/Engine/Source/Programs/SubmitTool/Private/Logic/P4LockdownService.cpp
@@ -109,6 +109,8 @@ FPathLockdownResult FP4LockdownService::ArePathsInLockdown(const TArray<FString>
 		AllowListState = bShouldParseDownloadedData ? EAllowListState::Downloaded : EAllowListState::Cached;
 	}
 
+	TSet<FString> UniqueGroups;
+	TSet<FString> UniqueMessages;
 	// Finally validate the paths
 	for(const FString& Path : InPaths)
 	{
@@ -120,20 +122,15 @@ FPathLockdownResult FP4LockdownService::ArePathsInLockdown(const TArray<FString>
 			GlobalResult.bAllowlisted &= Result.bAllowlisted;
 			GlobalResult.bIsLocked = Result.bIsLocked;
 
-			if (!GlobalResult.Message.IsEmpty())
-			{
-				GlobalResult.Message.Append(TEXT("\n"));
-			}
-			GlobalResult.Message.Append(Result.Message);
-
-			if (!GlobalResult.GroupName.IsEmpty())
-			{
-				GlobalResult.GroupName.Append(TEXT(" "));
-			}
-			GlobalResult.GroupName.Append(Result.GroupName);
+			UniqueGroups.Add(Result.GroupName);
+			UniqueMessages.Add(Result.Message);
 		}
 	}
 
+	GlobalResult.GroupName = FString::Join(UniqueGroups, TEXT(" - "));
+	GlobalResult.Message = FString::Join(UniqueMessages, TEXT("\n"));
+
+
 	return GlobalResult;
 }
 
