SubmitTool: Improve internal source control api to prevent callers from waiting on GT tasks (since the support for that in TaskGraph is limited).
Separate callback use case from task with return result use case.
Fix one example scenario in SwarmService.

#rnx

diff --git a/Engine/Source/Programs/SubmitTool/Private/Logic/Services/Interfaces/ISTSourceControlService.h b/Engine/Source/Programs/SubmitTool/Private/Logic/Services/Interfaces/ISTSourceControlService.h
--- a/Engine/Source/Programs/SubmitTool/Private/Logic/Services/Interfaces/ISTSourceControlService.h
+++ b/Engine/Source/Programs/SubmitTool/Private/Logic/Services/Interfaces/ISTSourceControlService.h
@@ -125,26 +125,40 @@ public:
 	virtual const FAuthTicket& GetAuthTicket() = 0;
 	virtual bool IsBusy() const = 0;
 
-	UE::Tasks::TTask<TSharedRef<FSCCommand>> RunCommand(FString&& InCommand, TArray<FString>&& InAdditionalArgs, FOnSCCCommandComplete InCompleteCallback = nullptr)
+	// Runs the command in a task thread and the provided callback in a dependent GT task.
+	// Note: The return type should be changed to void to prevent the caller from waiting on the GT task.
+	UE::Tasks::FTask RunCommand(FString&& InCommand, TArray<FString>&& InAdditionalArgs, FOnSCCCommandComplete InCompleteCallback)
 	{
 		TSharedRef<FSCCommand> Command = MakeShared<FSCCommand>(MoveTemp(InCommand), MoveTemp(InAdditionalArgs));
-		UE::Tasks::TTask<TSharedRef<FSCCommand>> CommandTask = UE::Tasks::Launch(UE_SOURCE_LOCATION, [this, Command]
+		UE::Tasks::FTask CommandTask = UE::Tasks::Launch(UE_SOURCE_LOCATION, [this, Command]
 		{
 			Command->bSuccess = RunCommandInternal(*Command);
-			return Command;
 		});
 
-		// launch a dependent task to run the callback on the GT
-		if (InCompleteCallback.IsBound())
+		return UE::Tasks::Launch(UE_SOURCE_LOCATION, [InCompleteCallback, Command]
 		{
-			return UE::Tasks::Launch(UE_SOURCE_LOCATION, [InCompleteCallback, Command]
-			{
-				TRACE_CPUPROFILER_EVENT_SCOPE_TEXT(*FString::Printf(TEXT("RunP4Command_%s_Callback"), *Command->Command));
-				InCompleteCallback.Execute(*Command);
-				return Command;
-			}, CommandTask, LowLevelTasks::ETaskPriority::Normal, UE::Tasks::EExtendedTaskPriority::GameThreadNormalPri);
-		}
-		return CommandTask;
+			TRACE_CPUPROFILER_EVENT_SCOPE_TEXT(*FString::Printf(TEXT("RunP4Command_%s_Callback"), *Command->Command));
+			InCompleteCallback.Execute(*Command);
+		}, CommandTask, LowLevelTasks::ETaskPriority::Normal, UE::Tasks::EExtendedTaskPriority::GameThreadNormalPri);
+	}
+
+	// Runs the command in a task thread, calling GetResult() will wait until the task is completed and return the command data.
+	UE::Tasks::TTask<TSharedRef<FSCCommand>> RunCommand(FString&& InCommand, TArray<FString>&& InAdditionalArgs)
+	{
+		TSharedRef<FSCCommand> Command = MakeShared<FSCCommand>(MoveTemp(InCommand), MoveTemp(InAdditionalArgs));
+		return UE::Tasks::Launch(UE_SOURCE_LOCATION, [this, Command]
+		{
+			Command->bSuccess = RunCommandInternal(*Command);
+			return Command;
+		});
+	}
+
+	// Runs the command synchronously and returns the command data
+	FSCCommand RunSyncCommand(FString&& InCommand, TArray<FString>&& InAdditionalArgs)
+	{
+		FSCCommand Command(MoveTemp(InCommand), MoveTemp(InAdditionalArgs));
+		Command.bSuccess = RunCommandInternal(Command);
+		return MoveTemp(Command);
 	}
 
 protected:
diff --git a/Engine/Source/Programs/SubmitTool/Private/Logic/SwarmService.cpp b/Engine/Source/Programs/SubmitTool/Private/Logic/SwarmService.cpp
--- a/Engine/Source/Programs/SubmitTool/Private/Logic/SwarmService.cpp
+++ b/Engine/Source/Programs/SubmitTool/Private/Logic/SwarmService.cpp
@@ -18,7 +18,7 @@
 FSwarmService::FSwarmService(TWeakPtr<FSubmitToolServiceProvider> InServiceProvider) :
 	ServiceProvider(InServiceProvider)
 {
-	LaunchGetSwarmURLTask();
+	GetSwarmURLTask = ServiceProvider.Pin()->GetService<ISTSourceControlService>()->RunCommand(TEXT("property"), { "-l", "-n", "P4.Swarm.URL" });
 	ServiceProvider.Pin()->GetService<ISTSourceControlService>()->GetUsersAndGroups(FOnUsersAndGroupsGet::FDelegate::CreateLambda([this](TArray<TSharedPtr<FUserData>>& InUsers, TArray<TSharedPtr<FString>>& InGroups) {
 		Users = &InUsers;
 		Groups = &InGroups;
@@ -283,25 +283,20 @@ void FSwarmService::UpdateReviewDescription(const TDelegate<void(bool)>& OnCompl
 
 }
 
-void FSwarmService::LaunchGetSwarmURLTask()
+FString FSwarmService::ReviewsURL()
 {
-	GetSwarmURLTask = ServiceProvider.Pin()->GetService<ISTSourceControlService>()->RunCommand(TEXT("property"), { "-l", "-n", "P4.Swarm.URL" },
-		FOnSCCCommandComplete::CreateLambda(
-		[this](const FSCCommand& Cmd)
+	if (GetSwarmURLTask.IsValid())
+	{
+		TSharedRef<FSCCommand> Cmd = GetSwarmURLTask.GetResult();
+		if (Cmd->bSuccess && Cmd->Values.Num() > 0)
 		{
-			if (Cmd.bSuccess && Cmd.Values.Num() > 0)
+			if (const FString* Value = Cmd->Values[0].Find(TEXT("value")))
 			{
-				if (const FString* Value = Cmd.Values[0].Find(TEXT("value")))
-				{
-					SwarmURL = *Value;
-				}
+				SwarmURL = *Value;
 			}
-		}));
-}
-
-FString FSwarmService::ReviewsURL() const
-{
-	GetSwarmURLTask.Wait();
+		}
+		GetSwarmURLTask = {};
+	}
 	return SwarmURL.IsEmpty() ? FString{} : SwarmURL / TEXT("api/v9/reviews");
 }
 
diff --git a/Engine/Source/Programs/SubmitTool/Private/Logic/SwarmService.h b/Engine/Source/Programs/SubmitTool/Private/Logic/SwarmService.h
--- a/Engine/Source/Programs/SubmitTool/Private/Logic/SwarmService.h
+++ b/Engine/Source/Programs/SubmitTool/Private/Logic/SwarmService.h
@@ -59,16 +59,14 @@ public:
 	const TArray<FString> GetUsersInSwarmTag() const;
 
 private:
-	void LaunchGetSwarmURLTask();
-	FString ReviewsURL() const;
-
+	FString ReviewsURL();
 
 private:
 	TSharedPtr<IHttpRequest> CreateSwarmRequest;
 	TSharedPtr<IHttpRequest> UpdateSwarmRequest;
 
 	TWeakPtr<FSubmitToolServiceProvider> ServiceProvider;
-	UE::Tasks::FTask GetSwarmURLTask;
+	UE::Tasks::TTask<TSharedRef<FSCCommand>> GetSwarmURLTask;
 	FString SwarmURL;
 	TUniquePtr<FSwarmReview> Review = nullptr;
 
