SubmitTool: Fix rare use-after-free in SubmitToolPerforce.
Make a copy of GroupsToExclude for the FGroupsClientUser, else if the task is still running on shutdown it will access the free'd GeneralParameters.GroupsToExclude array).
Found by running with -stompmalloc (and a slow GetGroups call at 26s instead of the normal 1s).

#rnx

diff --git a/Engine/Source/Programs/SubmitTool/Private/Logic/Services/SourceControl/SubmitToolPerforce.cpp b/Engine/Source/Programs/SubmitTool/Private/Logic/Services/SourceControl/SubmitToolPerforce.cpp
--- a/Engine/Source/Programs/SubmitTool/Private/Logic/Services/SourceControl/SubmitToolPerforce.cpp
+++ b/Engine/Source/Programs/SubmitTool/Private/Logic/Services/SourceControl/SubmitToolPerforce.cpp
@@ -541,13 +541,16 @@ void FSubmitToolPerforce::GetGroups(const FOnGroupsGet::FDelegate& Callback)
 					}
 
 					TArray<TSharedPtr<FString>>& GroupsArray;
-					const TArray<FString>& ExcludedGroups;
+					const TArray<FString> ExcludedGroups;
 				};
 
 				FGroupsClientUser P4User(ResultValues, bIsUnicodeServer ? EP4ClientUserFlags::UnicodeServer : EP4ClientUserFlags::None, ResultInfo, CachedGroupsArray, Parameters.GeneralParameters.GroupsToExclude);
 
-				UE_LOG(LogSubmitToolP4Debug, Log, TEXT("Running command: p4 -p %s -u %s -c %s -ztag groups"), TO_TCHAR(P4Client.GetPort().Text(), bIsUnicodeServer), *UserName, *Client);
-				P4Client.Run(FROM_TCHAR(TEXT("groups"), bIsUnicodeServer), &P4User);
+				{
+					TRACE_CPUPROFILER_EVENT_SCOPE(FSubmitToolPerforce::RunP4Command_groups);
+					UE_LOG(LogSubmitToolP4Debug, Log, TEXT("Running command: p4 -p %s -u %s -c %s -ztag groups"), TO_TCHAR(P4Client.GetPort().Text(), bIsUnicodeServer), *UserName, *Client);
+					P4Client.Run(FROM_TCHAR(TEXT("groups"), bIsUnicodeServer), &P4User);
+				}
 
 				for(const FText& Msg : ResultInfo.InfoMessages)
 				{
@@ -566,6 +569,7 @@ void FSubmitToolPerforce::GetGroups(const FOnGroupsGet::FDelegate& Callback)
 
 				AsyncTask(ENamedThreads::GameThread, [this] 
 				{
+					TRACE_CPUPROFILER_EVENT_SCOPE(FSubmitToolPerforce::GetGroups_Broadcast);
 					OnGroupsGetCallbacks.Broadcast(CachedGroupsArray);
 					OnGroupsGetCallbacks.Clear(); 
 				});
@@ -578,6 +582,7 @@ void FSubmitToolPerforce::GetGroups(const FOnGroupsGet::FDelegate& Callback)
 
 void FSubmitToolPerforce::LoadRecentGroups()
 {
+	TRACE_CPUPROFILER_EVENT_SCOPE(FSubmitToolPerforce::LoadRecentGroups);
 	const TUniquePtr<FAES::FAESKey>& Key = FCredentialsService::GetEncryptionKey();
 
 	if(!Key.IsValid())
