SubmitTool: Tiny bugfix in GetGroups, use correct casing for "isSubGroup".
Tiny optimization to early out when bIsSubGroup is true.

#rnx

diff --git a/Engine/Source/Programs/SubmitTool/Private/Logic/Services/SourceControl/SubmitToolPerforce.cpp b/Engine/Source/Programs/SubmitTool/Private/Logic/Services/SourceControl/SubmitToolPerforce.cpp
--- a/Engine/Source/Programs/SubmitTool/Private/Logic/Services/SourceControl/SubmitToolPerforce.cpp
+++ b/Engine/Source/Programs/SubmitTool/Private/Logic/Services/SourceControl/SubmitToolPerforce.cpp
@@ -502,42 +502,46 @@ void FSubmitToolPerforce::GetGroups(const FOnGroupsGet::FDelegate& Callback)
 
 					virtual void OutputStat(StrDict* VarList) override
 					{
+						TRACE_CPUPROFILER_EVENT_SCOPE(FGroupsClientUser::OutputStat);
 						StrRef Var, Value;
 
 						FString Group;
-						bool bIsUser = false;
+						int32 FoundKeys = 0;
 						bool bIsSubGroup = false;
 						// Iterate over each variable and add to records
-						for(int32 Index = 0; VarList->GetVar(Index, Var, Value); Index++)
+						for (int32 Index = 0; FoundKeys < 2 && VarList->GetVar(Index, Var, Value); Index++)
 						{
-							if(Var == "isSubgroup")
+							if (Var == "isSubGroup")
 							{
 								bIsSubGroup = Value != "0";
+								++FoundKeys;
 							}
-
-							if(Var == "group")
+							else if (Var == "group")
 							{
 								Group = TO_TCHAR(Value.Text(), IsUnicodeServer());
+								++FoundKeys;
 							}
 						}
 
-						if(!Group.IsEmpty() && !GroupsArray.ContainsByPredicate([&Group](const TSharedPtr<FString>& InStr) { return InStr->Equals(Group, ESearchCase::IgnoreCase); }) && !bIsSubGroup)
+						if (bIsSubGroup || Group.IsEmpty())
 						{
-							bool bIsExcluded = false;
-							for(const FString& Filter : ExcludedGroups)
-							{
-								if(Group.StartsWith(Filter))
-								{
-									bIsExcluded = true;
-									break;
-								}
-							}
+							return;
+						}
 
-							if(!bIsExcluded)
+						if (GroupsArray.ContainsByPredicate([&Group](const TSharedPtr<FString>& InStr) { return InStr->Equals(Group, ESearchCase::IgnoreCase); }))
+						{
+							return;
+						}
+
+						for (const FString& Filter : ExcludedGroups)
+						{
+							if (Group.StartsWith(Filter))
 							{
-								GroupsArray.Add(MakeShared<FString>(Group));
+								return;
 							}
 						}
+
+						GroupsArray.Add(MakeShared<FString>(Group));
 					}
 
 					TArray<TSharedPtr<FString>>& GroupsArray;
@@ -549,7 +553,7 @@ void FSubmitToolPerforce::GetGroups(const FOnGroupsGet::FDelegate& Callback)
 				{
 					TRACE_CPUPROFILER_EVENT_SCOPE(FSubmitToolPerforce::RunP4Command_groups);
 					UE_LOG(LogSubmitToolP4Debug, Log, TEXT("Running command: p4 -p %s -u %s -c %s -ztag groups"), TO_TCHAR(P4Client.GetPort().Text(), bIsUnicodeServer), *UserName, *Client);
-					P4Client.Run(FROM_TCHAR(TEXT("groups"), bIsUnicodeServer), &P4User);
+					P4Client.Run("groups", &P4User);
 				}
 
 				for(const FText& Msg : ResultInfo.InfoMessages)
