Fixing Waits for tasks that were scheduled in the GT for synchronous p4 commands that were causing deadlocks
Fixed potential issues with mutex/missing std::atomic in SubmitToolPerforce
Only use last known shelf state from submit tool for avoiding the p4 diff when we certainly know the files are up to date, use p4 diff otherwise to avoid user confusion

diff --git a/Engine/Source/Programs/SubmitTool/Private/Logic/ChangelistService.cpp b/Engine/Source/Programs/SubmitTool/Private/Logic/ChangelistService.cpp
--- a/Engine/Source/Programs/SubmitTool/Private/Logic/ChangelistService.cpp
+++ b/Engine/Source/Programs/SubmitTool/Private/Logic/ChangelistService.cpp
@@ -48,7 +48,7 @@ void FChangelistService::Init()
 	TRACE_CPUPROFILER_EVENT_SCOPE(FChangelistService::Init);
 	if(SourceControlService->GetProvider().IsValid())
 	{
-		InitialP4RunTasks.Add(LaunchGetOpenedFilesDepotPaths());
+		InitialP4OpenedFilesTask = LaunchGetOpenedFilesDepotPaths();
 
 		CachedSCCProvider = SourceControlService->GetProvider().Get();
 
@@ -59,10 +59,9 @@ void FChangelistService::Init()
 		else
 		{
 			// Only do a p4 describe if we're not using the default CL
-			InitialP4RunTasks.Add(LaunchGetShelvedFilesDepotPaths());
+			InitialP4ShelvedFilesTask = LaunchGetShelvedFilesDepotPaths();
 			FindInitialChangelistsAsync();
 		}
-		
 	}
 	else
 	{
@@ -702,62 +701,70 @@ void FChangelistService::PrintFilesAndShelvedFiles() const
 		}
 	}
 }
+UE::Tasks::TTask<TSharedRef<FSCCommand>> FChangelistService::LaunchGetOpenedFilesDepotPaths() const
+{
+	return SourceControlService->RunCommand(TEXT("opened"), { });
+}
 
-UE::Tasks::FTask FChangelistService::LaunchGetOpenedFilesDepotPaths()
+void FChangelistService::ParseOpenedFilesDepotPaths(TSharedRef<FSCCommand> InCompletedSCCommand)
 {
-	return SourceControlService->RunCommand(TEXT("opened"), { }, FOnSCCCommandComplete::CreateLambda(
-		[this](const FSCCommand& Cmd)
+	FilesDepotPaths.Reset();
+	OpenFilesInChangelists.Reset();
+	if (InCompletedSCCommand->bSuccess)
+	{
+		for (const TMap<FString, FString>& Record : InCompletedSCCommand->Values)
 		{
-			FilesDepotPaths.Reset();
-			OpenFilesInChangelists.Reset();
-			if (Cmd.bSuccess)
+			const FString* DepotFile = Record.Find(TEXT("depotFile"));
+			const FString* Change = Record.Find(TEXT("change"));
+			if (DepotFile && Change)
 			{
-				for(const TMap<FString, FString>& Record : Cmd.Values)
+				if (*Change == CLID)
 				{
-					const FString* DepotFile = Record.Find(TEXT("depotFile"));
-					const FString* Change = Record.Find(TEXT("change"));
-					if (DepotFile && Change)
-					{
-						if (*Change == CLID)
-						{
-							FilesDepotPaths.Add(*DepotFile);
-						}
-						OpenFilesInChangelists.Add(*DepotFile, *Change);
-					}
+					FilesDepotPaths.Add(*DepotFile);
 				}
+				OpenFilesInChangelists.Add(*DepotFile, *Change);
 			}
-		}));
+		}
+	}
+}
+UE::Tasks::TTask<TSharedRef<FSCCommand>> FChangelistService::LaunchGetShelvedFilesDepotPaths() const
+{
+	return SourceControlService->RunCommand(TEXT("describe"), { "-S", CLID });
 }
 
-UE::Tasks::FTask FChangelistService::LaunchGetShelvedFilesDepotPaths()
+void FChangelistService::ParseShelvedFilesDepotPaths(TSharedRef<FSCCommand> InCompletedSCCommand)
 {
-	return SourceControlService->RunCommand(TEXT("describe"), { "-S", CLID }, FOnSCCCommandComplete::CreateLambda(
-	[this](const FSCCommand& Cmd)
+	ShelvedFilesDepotPaths.Reset();
+	if (InCompletedSCCommand->bSuccess)
 	{
-		if (Cmd.bSuccess)
+		for (const TMap<FString, FString>& Record : InCompletedSCCommand->Values)
 		{
-			for (const TMap<FString, FString>& Record : Cmd.Values)
+			int32 I = 0;
+			FString PathsKey = FString::Printf(TEXT("depotFile%d"), I);
+			while (Record.Contains(PathsKey))
 			{
-				int32 I = 0;
-				FString PathsKey = FString::Printf(TEXT("depotFile%d"), I);
-				while(Record.Contains(PathsKey))
-				{
-					ShelvedFilesDepotPaths.Add(Record[PathsKey]);
-					++I;
-					PathsKey = FString::Printf(TEXT("depotFile%d"), I);
-				}
+				ShelvedFilesDepotPaths.Add(Record[PathsKey]);
+				++I;
+				PathsKey = FString::Printf(TEXT("depotFile%d"), I);
 			}
 		}
-	}));
+	}
 }
 
 void FChangelistService::WaitForInitialP4RunTasks()
 {
-	if (InitialP4RunTasks.Num())
+	TRACE_CPUPROFILER_EVENT_SCOPE(FChangelistService::WaitForInitialP4RunTasks);
+
+	if (InitialP4OpenedFilesTask.IsValid())
+	{
+		ParseOpenedFilesDepotPaths(InitialP4OpenedFilesTask.GetResult());
+		InitialP4OpenedFilesTask = {};
+	}
+
+	if (InitialP4ShelvedFilesTask.IsValid())
 	{
-		TRACE_CPUPROFILER_EVENT_SCOPE(FChangelistService::WaitForInitialP4RunTasks);
-		UE::Tasks::Wait(InitialP4RunTasks);
-		InitialP4RunTasks.Empty();
+		ParseShelvedFilesDepotPaths(InitialP4ShelvedFilesTask.GetResult());
+		InitialP4ShelvedFilesTask = {};
 	}
 }
 
@@ -769,7 +776,7 @@ const TArray<FString>& FChangelistService::GetOpenedFilesDepotPaths()
 	if (FilesDepotPaths.IsEmpty() && !FilesInCL.IsEmpty())
 	{
 		TRACE_CPUPROFILER_EVENT_SCOPE(FChangelistService::GetOpenedFilesDepotPaths);
-		LaunchGetOpenedFilesDepotPaths().Wait();
+		ParseOpenedFilesDepotPaths(LaunchGetOpenedFilesDepotPaths().GetResult());
 	}
 
 	return FilesDepotPaths;
@@ -790,7 +797,7 @@ const TArray<FString>& FChangelistService::GetShelvedFilesDepotPaths()
 	{
 		TRACE_CPUPROFILER_EVENT_SCOPE(FChangelistService::GetShelvedFilesDepotPaths);
 		ShelvedFilesDepotPaths.Reset();
-		LaunchGetShelvedFilesDepotPaths().Wait();
+		ParseShelvedFilesDepotPaths(LaunchGetShelvedFilesDepotPaths().GetResult());
 	}
 
 	return ShelvedFilesDepotPaths;
@@ -804,7 +811,7 @@ const TMap<FString, FString>& FChangelistService::GetOpenFilesInChangelists()
 	if (OpenFilesInChangelists.IsEmpty())
 	{
 		TRACE_CPUPROFILER_EVENT_SCOPE(FChangelistService::GetOpenFilesInChangelists);
-		LaunchGetOpenedFilesDepotPaths().Wait();
+		ParseOpenedFilesDepotPaths(LaunchGetOpenedFilesDepotPaths().GetResult());
 	}
 
 	return OpenFilesInChangelists;
@@ -823,9 +830,7 @@ void FChangelistService::CheckShelfIsUpToDate(TFunction<void(bool, TArray<FStrin
 		return;
 	}
 	
-	// If submit tool has ever known the shelf state for this CL, use that instead of running a full diff which might take too long.
-	// This unfortunately will be incorrect in the cases that files were manually shelved in p4v and will do an unnecessary re-shelve but 
-	// still better than paying for a diff every time just to avoid this edge case
+	// If submit tool has ever known the shelf state for this CL, use that datetime and check if we know that local files have not changed since
 	TSharedPtr<ICacheDataService> CacheService = ServiceProvider->GetService<ICacheDataService>();
 	if (CacheService)
 	{
@@ -844,24 +849,29 @@ void FChangelistService::CheckShelfIsUpToDate(TFunction<void(bool, TArray<FStrin
 				}
 			}
 			
-			bIsShelfReady = FilePathsOutOfDate.Num() == 0;
-			OnComplete(bIsShelfReady, MoveTemp(FilePathsOutOfDate));
-			return;
+			// If local files have not changed, we can early out without running p4 diff
+			if (FilePathsOutOfDate.Num() == 0)
+			{
+				bIsShelfReady = true;
+				OnComplete(bIsShelfReady, MoveTemp(FilePathsOutOfDate));
+				return;
+			}
 		}
 	}
 
-	// Only do a real p4 diff if shelf state is unknown to submit tool 
-	UE_LOG(LogSubmitToolP4, Log, TEXT("Checking Shelf diff with local files..."));
-	TArray<FString> DiffArgs = GetOpenedFilesDepotPaths();
-	for (FString& FilePath : DiffArgs)
-	{
-		FilePath.Appendf(TEXT("@=%s"), *CLID);
-	}
-	DiffArgs.Insert(TEXT("-ds"), 0);
-	DiffArgs.Insert(TEXT("-Od"), 0);
-	if (!DiffTask.IsValid() || DiffTask.IsCompleted())
+	if (!bDiffTaskRunning)
 	{
-		DiffTask = SourceControlService->RunCommand(TEXT("diff"), MoveTemp(DiffArgs), FOnSCCCommandComplete::CreateLambda(
+		bDiffTaskRunning = true;
+		// Real p4 diff
+		UE_LOG(LogSubmitToolP4, Log, TEXT("Checking Shelf diff with local files..."));
+		TArray<FString> DiffArgs = GetOpenedFilesDepotPaths();
+		for (FString& FilePath : DiffArgs)
+		{
+			FilePath.Appendf(TEXT("@=%s"), *CLID);
+		}
+		DiffArgs.Insert(TEXT("-ds"), 0);
+		DiffArgs.Insert(TEXT("-Od"), 0);
+		SourceControlService->RunCommand(TEXT("diff"), MoveTemp(DiffArgs), FOnSCCCommandComplete::CreateLambda(
 			[OnComplete, this](const FSCCommand& Cmd)
 			{
 				TRACE_CPUPROFILER_EVENT_SCOPE(FChangelistService::CheckShelfIsUpToDate_DiffResult);
@@ -898,6 +908,7 @@ void FChangelistService::CheckShelfIsUpToDate(TFunction<void(bool, TArray<FStrin
 				}
 
 				OnComplete(bIsShelfReady, FilesWithChanges);
+				bDiffTaskRunning = false;
 			}));
 	}
 }
@@ -913,7 +924,6 @@ void FChangelistService::EnsureShelfIsCurrent()
 			// if shelf is not up to date, we need to refresh it
 			if (!bIsUpToDate)
 			{
-
 				EDialogFactoryResult DialogResult = EDialogFactoryResult::Confirm;
 				if (!FSubmitToolUserPrefs::Get()->bAutoStompShelf)
 				{
diff --git a/Engine/Source/Programs/SubmitTool/Private/Logic/ChangelistService.h b/Engine/Source/Programs/SubmitTool/Private/Logic/ChangelistService.h
--- a/Engine/Source/Programs/SubmitTool/Private/Logic/ChangelistService.h
+++ b/Engine/Source/Programs/SubmitTool/Private/Logic/ChangelistService.h
@@ -13,6 +13,7 @@
 #include "Tasks/Task.h"
 
 struct FSourceControlResultInfo;
+struct FSCCommand;
 class FSubmitToolServiceProvider;
 class ISTSourceControlService;
 using FSourceControlOperationRef = TSharedRef<class ISourceControlOperation, ESPMode::ThreadSafe>;
@@ -139,7 +140,8 @@ private:
 	const TSharedPtr<ISTSourceControlService> SourceControlService;
 
 	TArray<TSharedRef<ISourceControlOperation>> ActiveP4Operations;
-	TArray<UE::Tasks::FTask> InitialP4RunTasks;
+	UE::Tasks::TTask<TSharedRef<FSCCommand>> InitialP4OpenedFilesTask;
+	UE::Tasks::TTask<TSharedRef<FSCCommand>> InitialP4ShelvedFilesTask;
 
 	FSourceControlChangelistPtr ChangelistPtr = nullptr;
 	FSourceControlChangelistStatePtr ChangelistState = nullptr;
@@ -154,12 +156,15 @@ private:
 	// Map filled once by GetOpenedFilesDepotPaths which is more efficient than a FUpdatePendingChangelistsStatus operation for all states for all pending changelists.
 	TMap<FString, FString> OpenFilesInChangelists;
 	bool bIsShelfReady = false;
-	UE::Tasks::FTask DiffTask;
+	bool bDiffTaskRunning = false;
 	void FindInitialChangelistsAsync();
 
 	void WaitForInitialP4RunTasks();
-	UE::Tasks::FTask LaunchGetOpenedFilesDepotPaths();
-	UE::Tasks::FTask LaunchGetShelvedFilesDepotPaths();
+	UE::Tasks::TTask<TSharedRef<FSCCommand>> LaunchGetOpenedFilesDepotPaths() const;
+	UE::Tasks::TTask<TSharedRef<FSCCommand>> LaunchGetShelvedFilesDepotPaths() const;
+
+	void ParseOpenedFilesDepotPaths(TSharedRef<FSCCommand> InCompletedSCCommand);
+	void ParseShelvedFilesDepotPaths(TSharedRef<FSCCommand> InCompletedSCCommand);
 
 	void FetchChangelistCallback(const FSourceControlOperationRef& UpdateOperation, ECommandResult::Type Result);
 	void CacheChangelistState();
diff --git a/Engine/Source/Programs/SubmitTool/Private/Logic/Services/Interfaces/ISTSourceControlService.h b/Engine/Source/Programs/SubmitTool/Private/Logic/Services/Interfaces/ISTSourceControlService.h
--- a/Engine/Source/Programs/SubmitTool/Private/Logic/Services/Interfaces/ISTSourceControlService.h
+++ b/Engine/Source/Programs/SubmitTool/Private/Logic/Services/Interfaces/ISTSourceControlService.h
@@ -126,20 +126,19 @@ public:
 	virtual bool IsBusy() const = 0;
 
 	// Runs the command in a task thread and the provided callback in a dependent GT task.
-	// Note: The return type should be changed to void to prevent the caller from waiting on the GT task.
-	UE::Tasks::FTask RunCommand(FString&& InCommand, TArray<FString>&& InAdditionalArgs, FOnSCCCommandComplete InCompleteCallback)
+	void RunCommand(FString&& InCommand, TArray<FString>&& InAdditionalArgs, FOnSCCCommandComplete InCompleteCallback)
 	{
+		check(InCompleteCallback.IsBound());
 		TSharedRef<FSCCommand> Command = MakeShared<FSCCommand>(MoveTemp(InCommand), MoveTemp(InAdditionalArgs));
-		UE::Tasks::FTask CommandTask = UE::Tasks::Launch(UE_SOURCE_LOCATION, [this, Command]
+		UE::Tasks::FTask CommandTask = UE::Tasks::Launch(UE_SOURCE_LOCATION, [this, Command, InCompleteCallback]
 		{
 			Command->bSuccess = RunCommandInternal(*Command);
+			UE::Tasks::Launch(UE_SOURCE_LOCATION, [InCompleteCallback, Command]
+			{
+				TRACE_CPUPROFILER_EVENT_SCOPE_TEXT(*FString::Printf(TEXT("RunP4Command_%s_Callback"), *Command->Command));
+				InCompleteCallback.Execute(*Command);
+			}, LowLevelTasks::ETaskPriority::Normal, UE::Tasks::EExtendedTaskPriority::GameThreadNormalPri);
 		});
-
-		return UE::Tasks::Launch(UE_SOURCE_LOCATION, [InCompleteCallback, Command]
-		{
-			TRACE_CPUPROFILER_EVENT_SCOPE_TEXT(*FString::Printf(TEXT("RunP4Command_%s_Callback"), *Command->Command));
-			InCompleteCallback.Execute(*Command);
-		}, CommandTask, LowLevelTasks::ETaskPriority::Normal, UE::Tasks::EExtendedTaskPriority::GameThreadNormalPri);
 	}
 
 	// Runs the command in a task thread, calling GetResult() will wait until the task is completed and return the command data.
diff --git a/Engine/Source/Programs/SubmitTool/Private/Logic/Services/SourceControl/SubmitToolPerforce.cpp b/Engine/Source/Programs/SubmitTool/Private/Logic/Services/SourceControl/SubmitToolPerforce.cpp
--- a/Engine/Source/Programs/SubmitTool/Private/Logic/Services/SourceControl/SubmitToolPerforce.cpp
+++ b/Engine/Source/Programs/SubmitTool/Private/Logic/Services/SourceControl/SubmitToolPerforce.cpp
@@ -24,16 +24,14 @@ constexpr size_t MaxConnectionsAttempts = 50; // 50*0.2f => ~10s
 constexpr float ConnectionAttemptWaitTime = 0.2f;
 
 
-FP4Connection::FP4Connection(FClientApiWrapper& InConnection, FCriticalSection& InMutex)
-	: Connection(InConnection), P4ConnectionMutex(InMutex)
+FP4Connection::FP4Connection(FClientApiWrapper& InConnection)
+	: Connection(InConnection)
 {
-	FScopeLock Lock(&P4ConnectionMutex);
 	Connection.bIsReady = false;
 }
 
 FP4Connection::~FP4Connection()
 {
-	FScopeLock Lock(&P4ConnectionMutex);
 	Connection.bIsReady = true;
 }
 
@@ -61,7 +59,7 @@ TUniquePtr<FP4Connection> FConnectionPool::GetAvailableConnection()
 	{
 		if(ConWrapper->bIsReady)
 		{
-			return MakeUnique<FP4Connection>(*ConWrapper, Mutex);
+			return MakeUnique<FP4Connection>(*ConWrapper);
 		}
 	}
 
@@ -72,7 +70,7 @@ TUniquePtr<FP4Connection> FConnectionPool::GetAvailableConnection()
 		size_t Idx = CreateConnection();
 		if(Idx != -1)
 		{
-			return MakeUnique<FP4Connection>(*P4Connections[Idx], Mutex);
+			return MakeUnique<FP4Connection>(*P4Connections[Idx]);
 		}
 	}
 
@@ -104,7 +102,7 @@ size_t FConnectionPool::CreateConnection()
 	}
 
 	FScopeLock Lock(&Mutex);
-	return P4Connections.Add(MakeUnique<FClientApiWrapper>(MoveTemp(P4Client), FOnIsCancelled::CreateLambda([this] { return bWantsCancel; })));
+	return P4Connections.Add(MakeUnique<FClientApiWrapper>(MoveTemp(P4Client), FOnIsCancelled::CreateLambda([this] { return bWantsCancel.load(); })));
 }
 
 FSubmitToolPerforce::FSubmitToolPerforce(const FSubmitToolParameters& InParameters)
@@ -154,7 +152,12 @@ FSubmitToolPerforce::FSubmitToolPerforce(const FSubmitToolParameters& InParamete
 	}
 
 	TickHandle = FTSTicker::GetCoreTicker().AddTicker(FTickerDelegate::CreateRaw(this, &FSubmitToolPerforce::Tick));
-	GetStream();
+	UE::Tasks::TTask<TSharedRef<FSCCommand>> GetStreamTask = GetStream();
+
+	UE::Tasks::Launch(UE_SOURCE_LOCATION, [this, GetStreamTask] () mutable
+	{
+		ParseStreamCommand(*GetStreamTask.GetResult());
+	}, GetStreamTask, LowLevelTasks::ETaskPriority::Normal, UE::Tasks::EExtendedTaskPriority::GameThreadNormalPri);
 }
 
 FSubmitToolPerforce::~FSubmitToolPerforce()
@@ -175,7 +178,7 @@ const TUniquePtr<ISourceControlProvider>& FSubmitToolPerforce::GetProvider() con
 	return SCCProvider;
 }
 
-UE::Tasks::FTask FSubmitToolPerforce::GetStream(const FString& InStream, bool bRequestHierarchy)
+UE::Tasks::TTask<TSharedRef<FSCCommand>> FSubmitToolPerforce::GetStream(const FString& InStream)
 {
 	TArray<FString> Args = { TEXT("-o") };
 
@@ -184,75 +187,78 @@ UE::Tasks::FTask FSubmitToolPerforce::GetStream(const FString& InStream, bool bR
 		Args.Add(InStream);
 	}
 
-	return RunCommand(TEXT("stream"), MoveTemp(Args), FOnSCCCommandComplete::CreateLambda(
-		[this, bRequestHierarchy, InStream](const FSCCommand& Cmd)
+	return RunCommand(TEXT("stream"), MoveTemp(Args));
+}
+
+void FSubmitToolPerforce::ParseStreamCommand(const FSCCommand& InCompletedCommand, const FString& InRequestedStream, bool bParseHierarchy)
+{
+	if (InCompletedCommand.bSuccess && InCompletedCommand.Values.Num() > 0 && InCompletedCommand.Values[0].Contains(TEXT("Stream")))
+	{
+		TUniquePtr<FSCCStream> Stream = MakeUnique<FSCCStream>(InCompletedCommand.Values[0][TEXT("Stream")], InCompletedCommand.Values[0][TEXT("Parent")], InCompletedCommand.Values[0][TEXT("Type")]);
+
+		const FString Base = TEXT("Paths");
+		size_t i = 1;
+		FString PathsKey = FString::Printf(TEXT("%s%d"), *Base, i);
+		while (InCompletedCommand.Values[0].Contains(PathsKey))
 		{
-			if (Cmd.bSuccess && Cmd.Values.Num() > 0 && Cmd.Values[0].Contains(TEXT("Stream")))
+			FString Value = InCompletedCommand.Values[0][PathsKey];
+			if (Value.StartsWith(TEXT("import")))
 			{
-				TUniquePtr<FSCCStream> Stream = MakeUnique<FSCCStream>(Cmd.Values[0][TEXT("Stream")], Cmd.Values[0][TEXT("Parent")], Cmd.Values[0][TEXT("Type")]);
-
-				const FString Base = TEXT("Paths");
-				size_t i = 1;
-				FString PathsKey = FString::Printf(TEXT("%s%d"), *Base, i);
-				while (Cmd.Values[0].Contains(PathsKey))
+				int32 Idx = Value.Find(TEXT("//"));
+				if (Idx != INDEX_NONE)
 				{
-					FString Value = Cmd.Values[0][PathsKey];
-					if(Value.StartsWith(TEXT("import")))
-					{
-						Value = Value.RightChop(Value.Find(TEXT("//")));
-						while(!Value.EndsWith(TEXT("/")))
-						{
-							Value.RemoveAt(Value.Len() - 1);
-						}
-
-						Stream->AdditionalImportPaths.Add(Value);
-					}
+					Value = Value.RightChop(Value.Find(TEXT("//")));
+				}
 
-					++i;
-					PathsKey = FString::Printf(TEXT("%s%d"), *Base, i);
+				while (!Value.EndsWith(TEXT("/")))
+				{
+					Value.RemoveAt(Value.Len() - 1);
 				}
 
-				const FString Parent = Stream->Parent;
+				Stream->AdditionalImportPaths.Add(Value);
+			}
 
-				if (!Streams.Contains(Stream->Name))
-				{
-					if (bRequestHierarchy)
-					{
-						StreamHierarchy.Insert(Stream.Get(), 0);
-					}
+			++i;
+			PathsKey = FString::Printf(TEXT("%s%d"), *Base, i);
+		}
 
-					Streams.Add(Stream->Name, MoveTemp(Stream));
+		const FString Parent = Stream->Parent;
 
-					if (bRequestHierarchy && !Parent.IsEmpty() && !Parent.Equals(TEXT("none")))
-					{
-						GetStream(Parent);
-					}
-				}
-			}
-			else
+		if (!Streams.Contains(Stream->Name))
+		{
+			if (bParseHierarchy && !StreamHierarchy.ContainsByPredicate([&Stream](const FSCCStream* InStream) { return InStream->Name.Equals(Stream->Name, ESearchCase::IgnoreCase); }))
 			{
-				TUniquePtr<FSCCStream> Stream = MakeUnique<FSCCStream>(TEXT("Invalid"), FString(), TEXT("Invalid"));
+				StreamHierarchy.Insert(Stream.Get(), 0);
+			}
 
-				if (!InStream.IsEmpty())
-				{
-					if (!Streams.Contains(InStream))
-					{
-						Streams.Add(InStream, MoveTemp(Stream));
-					}
-				}
-				else
-				{
-					if (!Streams.Contains(TEXT("Invalid")))
-					{
-						StreamHierarchy.Add(Stream.Get());
-						Streams.Add(TEXT("Invalid"), MoveTemp(Stream));
-					}
-				}
+			Streams.Add(Stream->Name, MoveTemp(Stream));
+		}
+
+		if (bParseHierarchy && !Parent.IsEmpty() && !Parent.Equals(TEXT("none")))
+		{
+			ParseStreamCommand(RunSyncCommand(TEXT("stream"), { TEXT("-o"), Parent }), Parent, bParseHierarchy);
+		}
+	}
+	else
+	{
+		TUniquePtr<FSCCStream> Stream = MakeUnique<FSCCStream>(TEXT("Invalid"), FString(), TEXT("Invalid"));
+
+		if (!InRequestedStream.IsEmpty())
+		{
+			if (!Streams.Contains(InRequestedStream))
+			{
+				Streams.Add(InRequestedStream, MoveTemp(Stream));
 			}
 		}
-	));
+		else if (!Streams.Contains(TEXT("Invalid")))
+		{
+			StreamHierarchy.Add(Stream.Get());
+			Streams.Add(TEXT("Invalid"), MoveTemp(Stream));
+		}
+	}
 }
 
+
 void FSubmitToolPerforce::GetUsers(const FOnUsersGet::FDelegate& Callback)
 {
 	TRACE_CPUPROFILER_EVENT_SCOPE(FSubmitToolPerforce::GetUsers);
@@ -265,14 +271,15 @@ void FSubmitToolPerforce::GetUsers(const FOnUsersGet::FDelegate& Callback)
 
 	OnUsersGetCallbacks.Add(Callback);
 
-	if(!UserTask.IsValid())
+	if (!bUserTaskIsRunning)
 	{
-		UserTask = RunCommand(TEXT("users"), {}, FOnSCCCommandComplete::CreateLambda(
+		bUserTaskIsRunning = true;
+		RunCommand(TEXT("users"), {}, FOnSCCCommandComplete::CreateLambda(
 			[this](const FSCCommand& Cmd)
 			{
-				for(const TMap<FString, FString>& Record : Cmd.Values)
+				for (const TMap<FString, FString>& Record : Cmd.Values)
 				{
-					if(Record.Contains("User"))
+					if (Record.Contains("User"))
 					{
 						TSharedPtr<FUserData> User = MakeShared<FUserData>(Record[TEXT("User")], Record[TEXT("FullName")], Record[TEXT("Email")]);
 						CachedUsersArray.Add(User);
@@ -284,6 +291,7 @@ void FSubmitToolPerforce::GetUsers(const FOnUsersGet::FDelegate& Callback)
 
 				OnUsersGetCallbacks.Broadcast(CachedUsersArray);
 				OnUsersGetCallbacks.Clear();
+				bUserTaskIsRunning = false;
 			}
 		));
 	}
@@ -456,9 +464,10 @@ void FSubmitToolPerforce::GetGroups(const FOnGroupsGet::FDelegate& Callback)
 
 	OnGroupsGetCallbacks.Add(Callback);
 
-	if(!GroupTask.IsValid())
+	if(!bGroupTaskIsRunning)
 	{
-		GroupTask = UE::Tasks::Launch(UE_SOURCE_LOCATION,
+		bGroupTaskIsRunning = true;
+		UE::Tasks::Launch(UE_SOURCE_LOCATION,
 			[this] {
 				TRACE_CPUPROFILER_EVENT_SCOPE(FSubmitToolPerforce::GetGroups);
 				TUniquePtr<FP4Connection> Connection = Connections.GetAvailableConnection();
@@ -469,6 +478,7 @@ void FSubmitToolPerforce::GetGroups(const FOnGroupsGet::FDelegate& Callback)
 					TRACE_CPUPROFILER_EVENT_SCOPE(WaitForConnection);
 					if(Attempts >= MaxConnectionsAttempts)
 					{
+						bGroupTaskIsRunning = false;
 						return false;
 					}
 
@@ -574,12 +584,12 @@ void FSubmitToolPerforce::GetGroups(const FOnGroupsGet::FDelegate& Callback)
 				AsyncTask(ENamedThreads::GameThread, [this] 
 				{
 					TRACE_CPUPROFILER_EVENT_SCOPE(FSubmitToolPerforce::GetGroups_Broadcast);
+					bGroupTaskIsRunning = false;
 					OnGroupsGetCallbacks.Broadcast(CachedGroupsArray);
 					OnGroupsGetCallbacks.Clear(); 
 				});
 
 				return !ResultInfo.HasErrors();
-
 			});
 	}
 }
@@ -778,7 +788,7 @@ const FSCCStream* FSubmitToolPerforce::GetSCCStream(const FString& InStreamName)
 		}
 		if (I == 0)
 		{
-			GetStream(InStreamName, false).Wait();
+			ParseStreamCommand(*GetStream(InStreamName).GetResult(), InStreamName, false);
 		}
 	}
 	return nullptr;
@@ -795,7 +805,8 @@ const FString FSubmitToolPerforce::GetRootStreamName()
 		}
 		if (I == 0)
 		{
-			GetStream().Wait();
+			TSharedRef<FSCCommand> GetStreamResult = GetStream().GetResult();
+			ParseStreamCommand(*GetStreamResult);
 		}
 	}
 	return FString();
@@ -812,25 +823,33 @@ const FString FSubmitToolPerforce::GetCurrentStreamName()
 		}
 		if (I == 0)
 		{
-			GetStream().Wait();
+			ParseStreamCommand(*GetStream().GetResult());
 		}
 	}
 	return FString();
 }
 
-UE::Tasks::FTask FSubmitToolPerforce::LaunchGetDepotStreamDepth(const FString& InDepotName)
+int32 FSubmitToolPerforce::GetDepotStreamDepth(const FString& InDepotName)
 {
-	return RunCommand(TEXT("depot"), { TEXT("-o"), InDepotName }, FOnSCCCommandComplete::CreateLambda(
-		[this, InDepotName](const FSCCommand& Cmd)
+	TRACE_CPUPROFILER_EVENT_SCOPE(FSubmitToolPerforce::GetDepotStreamDepth);
+	for (int32 I = 0; I < 2; ++I)
+	{
+		if (int32* StreamDepth = DepotStreamDepths.Find(InDepotName))
 		{
-			if (Cmd.bSuccess && Cmd.Values.Num() > 0)
+			return *StreamDepth;
+		}
+		if (I == 0)
+		{
+			TSharedRef<FSCCommand> Cmd = RunCommand(TEXT("depot"), { TEXT("-o"), InDepotName }).GetResult();
+
+			if (Cmd->bSuccess && Cmd->Values.Num() > 0)
 			{
-				if (const FString* StreamDepth = Cmd.Values[0].Find(TEXT("StreamDepth")))
+				if (const FString* StreamDepth = Cmd->Values[0].Find(TEXT("StreamDepth")))
 				{
 					int32 Depth = 0;
-					for (int32 I = 2; I < StreamDepth->Len(); ++I)
+					for (int32 j = 2; j < StreamDepth->Len(); ++j)
 					{
-						if ((*StreamDepth)[I] == TCHAR('/'))
+						if ((*StreamDepth)[j] == TCHAR('/'))
 						{
 							Depth++;
 						}
@@ -839,21 +858,6 @@ UE::Tasks::FTask FSubmitToolPerforce::LaunchGetDepotStreamDepth(const FString& I
 					DepotStreamDepths.Add(InDepotName, Depth);
 				}
 			}
-		}));
-}
-
-int32 FSubmitToolPerforce::GetDepotStreamDepth(const FString& InDepotName)
-{
-	TRACE_CPUPROFILER_EVENT_SCOPE(FSubmitToolPerforce::GetDepotStreamDepth);
-	for (int32 I = 0; I < 2; ++I)
-	{
-		if (int32* StreamDepth = DepotStreamDepths.Find(InDepotName))
-		{
-			return *StreamDepth;
-		}
-		if (I == 0)
-		{
-			LaunchGetDepotStreamDepth(InDepotName).Wait();
 		}
 	}
 	return 0;
diff --git a/Engine/Source/Programs/SubmitTool/Private/Logic/Services/SourceControl/SubmitToolPerforce.h b/Engine/Source/Programs/SubmitTool/Private/Logic/Services/SourceControl/SubmitToolPerforce.h
--- a/Engine/Source/Programs/SubmitTool/Private/Logic/Services/SourceControl/SubmitToolPerforce.h
+++ b/Engine/Source/Programs/SubmitTool/Private/Logic/Services/SourceControl/SubmitToolPerforce.h
@@ -37,7 +37,7 @@ struct FClientApiWrapper
 		Connection->SetBreak(&KeepAlive);
 	}
 
-	bool bIsReady = true;
+	std::atomic<bool> bIsReady = true;
 	FP4KeepAlive KeepAlive;
 	TUniquePtr<ClientApi> Connection;
 };
@@ -45,14 +45,13 @@ struct FClientApiWrapper
 struct FP4Connection
 {
 public:
-	FP4Connection(FClientApiWrapper& InConnection, FCriticalSection& InMutex);
+	FP4Connection(FClientApiWrapper& InConnection);
 	~FP4Connection();
 
 	ClientApi& GetConnection() const { return *Connection.Connection; }
 
 private:
 	FClientApiWrapper& Connection;
-	FCriticalSection& P4ConnectionMutex;
 };
 
 class FConnectionPool
@@ -65,8 +64,8 @@ public:
 	void RequestCancel(){ bWantsCancel = true; }
 private:
 	size_t CreateConnection();
-	bool bWantsCancel = false;
-	bool bConnectionFailed = false;
+	std::atomic<bool> bWantsCancel = false;
+	std::atomic<bool> bConnectionFailed = false;
 	TArray<TUniquePtr<FClientApiWrapper>> P4Connections;
 	FCriticalSection Mutex;
 };
@@ -111,15 +110,15 @@ protected:
 
 private:
 	bool RunP4Command(FSCCommand& Command, EP4ClientUserFlags CmdFlags = DefaultFlags);
-	UE::Tasks::FTask LaunchGetDepotStreamDepth(const FString& InDepotName);
 
 	const FSubmitToolParameters& Parameters;
-	UE::Tasks::FTask GetStream(const FString& InStream = FString(), bool bRequestHierarchy = true);
+	UE::Tasks::TTask<TSharedRef<FSCCommand>> GetStream(const FString& InStream = FString());
+	void ParseStreamCommand(const FSCCommand& InCompletedCommand, const FString& InRequestedStream = FString(), bool bParseHierarchy = true);
 	TMap<FString, TSharedPtr<FUserData>> CachedUsers;
 	TArray<TSharedPtr<FUserData>> CachedUsersArray;
 	TArray<TSharedPtr<FString>> CachedGroupsArray;
-	UE::Tasks::FTask UserTask;
-	UE::Tasks::TTask<bool> GroupTask;
+	bool bUserTaskIsRunning = false;
+	std::atomic<bool> bGroupTaskIsRunning = false;
 	bool bIsUnicodeServer = false;
 	bool bUsersReady = false;
 	bool bGroupsReady = false;
