Moving thread unsafe operations in GetGroups to GT
Fixing potential use-after-free with ParseStreamCommand

diff --git a/Engine/Source/Programs/SubmitTool/Private/Logic/Services/SourceControl/SubmitToolPerforce.cpp b/Engine/Source/Programs/SubmitTool/Private/Logic/Services/SourceControl/SubmitToolPerforce.cpp
--- a/Engine/Source/Programs/SubmitTool/Private/Logic/Services/SourceControl/SubmitToolPerforce.cpp
+++ b/Engine/Source/Programs/SubmitTool/Private/Logic/Services/SourceControl/SubmitToolPerforce.cpp
@@ -479,6 +479,13 @@ void FSubmitToolPerforce::GetGroups(const FOnGroupsGet::FDelegate& Callback)
 					if(Attempts >= MaxConnectionsAttempts)
 					{
 						bGroupTaskIsRunning = false;
+						UE::Tasks::Launch(UE_SOURCE_LOCATION, [this]() mutable
+						{
+							LoadRecentGroups();
+							bGroupTaskIsRunning = false;
+							OnGroupsGetCallbacks.Broadcast(CachedGroupsArray);
+							OnGroupsGetCallbacks.Clear();
+						}, LowLevelTasks::ETaskPriority::Normal, UE::Tasks::EExtendedTaskPriority::GameThreadNormalPri);
 						return false;
 					}
 
@@ -558,7 +565,9 @@ void FSubmitToolPerforce::GetGroups(const FOnGroupsGet::FDelegate& Callback)
 					const TArray<FString> ExcludedGroups;
 				};
 
-				FGroupsClientUser P4User(ResultValues, bIsUnicodeServer ? EP4ClientUserFlags::UnicodeServer : EP4ClientUserFlags::None, ResultInfo, CachedGroupsArray, Parameters.GeneralParameters.GroupsToExclude);
+				TArray<TSharedPtr<FString>> GroupsResults;
+
+				FGroupsClientUser P4User(ResultValues, bIsUnicodeServer ? EP4ClientUserFlags::UnicodeServer : EP4ClientUserFlags::None, ResultInfo, GroupsResults, Parameters.GeneralParameters.GroupsToExclude);
 
 				{
 					TRACE_CPUPROFILER_EVENT_SCOPE(FSubmitToolPerforce::RunP4Command_groups);
@@ -579,15 +588,15 @@ void FSubmitToolPerforce::GetGroups(const FOnGroupsGet::FDelegate& Callback)
 					}
 				}
 
-				LoadRecentGroups();
-
-				AsyncTask(ENamedThreads::GameThread, [this] 
+				UE::Tasks::Launch(UE_SOURCE_LOCATION, [this, GroupsResults = MoveTemp(GroupsResults)]() mutable
 				{
 					TRACE_CPUPROFILER_EVENT_SCOPE(FSubmitToolPerforce::GetGroups_Broadcast);
+					CachedGroupsArray = MoveTemp(GroupsResults);
+					LoadRecentGroups();
 					bGroupTaskIsRunning = false;
 					OnGroupsGetCallbacks.Broadcast(CachedGroupsArray);
-					OnGroupsGetCallbacks.Clear(); 
-				});
+					OnGroupsGetCallbacks.Clear();
+				}, LowLevelTasks::ETaskPriority::Normal, UE::Tasks::EExtendedTaskPriority::GameThreadNormalPri);
 
 				return !ResultInfo.HasErrors();
 			});
@@ -788,7 +797,8 @@ const FSCCStream* FSubmitToolPerforce::GetSCCStream(const FString& InStreamName)
 		}
 		if (I == 0)
 		{
-			ParseStreamCommand(*GetStream(InStreamName).GetResult(), InStreamName, false);
+			TSharedRef<FSCCommand> GetStreamResult = GetStream(InStreamName).GetResult();
+			ParseStreamCommand(*GetStreamResult, InStreamName, false);
 		}
 	}
 	return nullptr;
@@ -823,7 +833,8 @@ const FString FSubmitToolPerforce::GetCurrentStreamName()
 		}
 		if (I == 0)
 		{
-			ParseStreamCommand(*GetStream().GetResult());
+			TSharedRef<FSCCommand> GetStreamResult = GetStream().GetResult();
+			ParseStreamCommand(*GetStreamResult);
 		}
 	}
 	return FString();
