Adding last synced CL tag taken from UGS

#Synced-CL 47420659

diff --git a/Engine/Source/Programs/SubmitTool/Private/Configuration/Configuration.cpp b/Engine/Source/Programs/SubmitTool/Private/Configuration/Configuration.cpp
--- a/Engine/Source/Programs/SubmitTool/Private/Configuration/Configuration.cpp
+++ b/Engine/Source/Programs/SubmitTool/Private/Configuration/Configuration.cpp
@@ -3,8 +3,10 @@
 #include "Configuration.h"
 
 #include "Logging/SubmitToolLog.h"
+#include "Misc/FileHelper.h"
 #include "Misc/Paths.h"
 #include "SubmitToolUtils.h"
+#include "Serialization/JsonSerializer.h"
 #include "CommandLine/CmdLineParameters.h"
 
 TSharedPtr<FConfiguration> FConfiguration::Instance;
@@ -69,6 +71,25 @@ void FConfiguration::Init()
 
 	Configuration.Values.Add(TEXT("$(SubmitToolConfigDir)"), *FPaths::ConvertRelativePathToFull(FPaths::ProjectConfigDir()));
 	
+	const FString UgsStateFile = RootDir / TEXT(".ugs/state.json");	
+	FString LastSyncedCL = TEXT("Unknown");
+
+	FString FileContents;
+	if (FPaths::FileExists(UgsStateFile) && FFileHelper::LoadFileToString(FileContents, *UgsStateFile))
+	{
+		const FString LastCLField = TEXT("LastSyncChangeNumber");
+		TSharedPtr<FJsonObject> RootJsonObject;
+		TSharedRef<TJsonReader<>> Reader = TJsonReaderFactory<>::Create(FileContents);
+		FJsonSerializer::Deserialize(Reader, RootJsonObject);
+		uint64 LastSyncedCLValue = 0;
+		if (RootJsonObject.IsValid() && RootJsonObject->TryGetNumberField(LastCLField, LastSyncedCLValue))
+		{
+			LastSyncedCL = LexToString(LastSyncedCLValue);
+		}
+	}
+
+	Configuration.Values.Add(TEXT("$(UGSSyncedCL)"), MoveTemp(LastSyncedCL));
+
 	Instance = MakeShared<FConfiguration>(Configuration);
 }
 
diff --git a/Engine/Source/Programs/SubmitTool/Private/Models/ModelInterface.cpp b/Engine/Source/Programs/SubmitTool/Private/Models/ModelInterface.cpp
--- a/Engine/Source/Programs/SubmitTool/Private/Models/ModelInterface.cpp
+++ b/Engine/Source/Programs/SubmitTool/Private/Models/ModelInterface.cpp
@@ -947,6 +947,8 @@ void FModelInterface::Submit()
 	FString Addendums;
 	AddendumAccumulator(ValidationService->GetAddendums(), CLDescription, Addendums);
 	AddendumAccumulator(PresubmitOperationsService->GetAddendums(), CLDescription, Addendums);
+	AddendumAccumulator(Parameters.GeneralParameters.AutomaticChangelistTags, CLDescription, Addendums);
+	Addendums = FConfiguration::Substitute(Addendums);
 
 	ChangelistService->Submit(Addendums, SubmitFinishedCallback);
 }
diff --git a/Engine/Source/Programs/SubmitTool/Private/Parameters/SubmitToolParameters.h b/Engine/Source/Programs/SubmitTool/Private/Parameters/SubmitToolParameters.h
--- a/Engine/Source/Programs/SubmitTool/Private/Parameters/SubmitToolParameters.h
+++ b/Engine/Source/Programs/SubmitTool/Private/Parameters/SubmitToolParameters.h
@@ -60,9 +60,12 @@ struct FGeneralParameters
 	
 	UPROPERTY()
 	TMap<FString, FString> PathOverrides;
-
+	
 	UPROPERTY()
 	TMap<FString, FString> PathOverridesInSourceBuild;
+	
+	UPROPERTY()
+	TArray<FString> AutomaticChangelistTags;
 };
 
 USTRUCT()
