Fixing download for lockdown files

#Synced-CL 47420659

diff --git a/Engine/Source/Programs/SubmitTool/Private/Logic/P4LockdownService.cpp b/Engine/Source/Programs/SubmitTool/Private/Logic/P4LockdownService.cpp
--- a/Engine/Source/Programs/SubmitTool/Private/Logic/P4LockdownService.cpp
+++ b/Engine/Source/Programs/SubmitTool/Private/Logic/P4LockdownService.cpp
@@ -56,23 +56,26 @@ bool FP4LockdownService::WaitForAllowListData()
 	}
 
 	// Write the downloaded files to disk and return true if at least one file is written
-	bool bSuccess = DownloadTask.IsCompleted() && DownloadTask.GetResult() && DownloadedFiles.Num() == Parameters.ConfigPaths.Num();
+	bool bSuccess = DownloadTask.IsCompleted() && DownloadTask.GetResult();
 	if (bSuccess)
 	{
 		bSuccess = false;
 		int32 I = 0;
 		for (const TPair<FString, FString>& Pair : Parameters.ConfigPaths)
 		{
-			FSharedBuffer FileBuffer = DownloadedFiles[I++];
-			const FString Path = GetFilePath(Pair.Key);
-			if (TUniquePtr<FArchive> File{IFileManager::Get().CreateFileWriter(*Path, FILEWRITE_EvenIfReadOnly)})
+			if (DownloadedFiles.Contains(Pair.Value))
 			{
-				File->Serialize(const_cast<void*>(FileBuffer.GetData()), FileBuffer.GetSize());
-				bSuccess = true;
-			}
-			else
-			{
-				UE_LOG(LogSubmitToolP4, Warning, TEXT("Couldn't create lockdown file %s"), *Path)
+				FSharedBuffer FileBuffer = DownloadedFiles[Pair.Value];
+				const FString Path = GetFilePath(Pair.Key);
+				if (TUniquePtr<FArchive> File{ IFileManager::Get().CreateFileWriter(*Path, FILEWRITE_EvenIfReadOnly) })
+				{
+					File->Serialize(const_cast<void*>(FileBuffer.GetData()), FileBuffer.GetSize());
+					bSuccess = true;
+				}
+				else
+				{
+					UE_LOG(LogSubmitToolP4, Warning, TEXT("Couldn't create lockdown file %s"), *Path)
+				}
 			}
 		}
 	}
diff --git a/Engine/Source/Programs/SubmitTool/Private/Logic/P4LockdownService.h b/Engine/Source/Programs/SubmitTool/Private/Logic/P4LockdownService.h
--- a/Engine/Source/Programs/SubmitTool/Private/Logic/P4LockdownService.h
+++ b/Engine/Source/Programs/SubmitTool/Private/Logic/P4LockdownService.h
@@ -77,7 +77,7 @@ private:
 	void GetAdditionalHardlockedPaths();
 	FString GetFilePath(const FString& InConfigId);
 
-	TArray<FSharedBuffer> DownloadedFiles;
+	TMap<FString, FSharedBuffer> DownloadedFiles;
 	UE::Tasks::TTask<bool> DownloadTask;
 	TWeakPtr<FSubmitToolServiceProvider> ServiceProvider;
 	const FP4LockdownParameters& Parameters;
diff --git a/Engine/Source/Programs/SubmitTool/Private/Logic/Services/Interfaces/ISTSourceControlService.h b/Engine/Source/Programs/SubmitTool/Private/Logic/Services/Interfaces/ISTSourceControlService.h
--- a/Engine/Source/Programs/SubmitTool/Private/Logic/Services/Interfaces/ISTSourceControlService.h
+++ b/Engine/Source/Programs/SubmitTool/Private/Logic/Services/Interfaces/ISTSourceControlService.h
@@ -86,7 +86,7 @@ struct FSCCommand
 	// out
 	FSCCRecordSet Values;
 	FSourceControlResultInfo Results;
-	TArray<FSharedBuffer> Data;
+	TMap<FString, FSharedBuffer> Data;
 	bool bSuccess = false;
 };
 
@@ -104,7 +104,7 @@ public:
 	virtual void GetGroups(const FOnGroupsGet::FDelegate& Callback) = 0;
 	virtual void GetUsersAndGroups(const FOnUsersAndGroupsGet::FDelegate& Callback) = 0;
 
-	virtual UE::Tasks::TTask<bool> DownloadFiles(TArray<FString>&& FilePaths, TArray<FSharedBuffer>& OutFileBuffers) = 0;
+	virtual UE::Tasks::TTask<bool> DownloadFiles(TArray<FString>&& FilePaths, TMap<FString, FSharedBuffer>& OutFileBuffers) = 0;
 
 	virtual bool IsAvailable() const = 0;
 
diff --git a/Engine/Source/Programs/SubmitTool/Private/Logic/Services/SourceControl/PerforceData.h b/Engine/Source/Programs/SubmitTool/Private/Logic/Services/SourceControl/PerforceData.h
--- a/Engine/Source/Programs/SubmitTool/Private/Logic/Services/SourceControl/PerforceData.h
+++ b/Engine/Source/Programs/SubmitTool/Private/Logic/Services/SourceControl/PerforceData.h
@@ -202,7 +202,8 @@ public:
 			{
 				if(File.IsFileComplete())
 				{
-					Files.Add(File.Release());
+					const FString Filename = File.GetDepotPath();
+					Files.Add(Filename, File.Release());
 				}
 			}
 			else
@@ -237,7 +238,8 @@ public:
 			{
 				if(File.IsFileComplete())
 				{
-					Files.Add(File.Release());
+					const FString Filename = File.GetDepotPath();
+					Files.Add(Filename, File.Release());
 				}
 				else
 				{
@@ -330,7 +332,7 @@ public:
 	}
 
 	/** Returns DataBuffer as a FSharedBuffer, note that once called DataBuffer will be empty */
-	inline TArray<FSharedBuffer> ReleaseData()
+	inline TMap<FString, FSharedBuffer> ReleaseData()
 	{
 		return MoveTemp(Files);
 	}
@@ -341,7 +343,7 @@ public:
 	FSourceControlResultInfo& ResultInfo;
 
 private:
-	TArray<FSharedBuffer> Files;
+	TMap<FString, FSharedBuffer> Files;
 
 	FP4File File;
 };
diff --git a/Engine/Source/Programs/SubmitTool/Private/Logic/Services/SourceControl/SubmitToolPerforce.cpp b/Engine/Source/Programs/SubmitTool/Private/Logic/Services/SourceControl/SubmitToolPerforce.cpp
--- a/Engine/Source/Programs/SubmitTool/Private/Logic/Services/SourceControl/SubmitToolPerforce.cpp
+++ b/Engine/Source/Programs/SubmitTool/Private/Logic/Services/SourceControl/SubmitToolPerforce.cpp
@@ -740,7 +740,7 @@ const FString FSubmitToolPerforce::GetRecentGroupsFilepath() const
 	return FPaths::Combine(FSubmitToolUtils::GetLocalAppDataPath(), TEXT("SubmitTool"), TEXT("recent_groups.dat"));
 }
 
-UE::Tasks::TTask<bool> FSubmitToolPerforce::DownloadFiles(TArray<FString>&& FilePaths, TArray<FSharedBuffer>& OutFileBuffers)
+UE::Tasks::TTask<bool> FSubmitToolPerforce::DownloadFiles(TArray<FString>&& FilePaths, TMap<FString, FSharedBuffer>& OutFileBuffers)
 {
 	FilePaths.Insert(TEXT("-q"), 0);
 	TSharedRef<FSCCommand> Command = MakeShared<FSCCommand>(TEXT("print"), MoveTemp(FilePaths));
@@ -749,11 +749,8 @@ UE::Tasks::TTask<bool> FSubmitToolPerforce::DownloadFiles(TArray<FString>&& File
 	{
 		if (RunP4Command(*Command, DefaultFlags | EP4ClientUserFlags::CollectData))
 		{
-			if (Command->Data.Num() == Command->Args.Num() - 1)
-			{
-				OutFileBuffers = MoveTemp(Command->Data);
-				return true;
-			}
+			OutFileBuffers = MoveTemp(Command->Data);
+			return true;
 		}
 		return false;
 	});
diff --git a/Engine/Source/Programs/SubmitTool/Private/Logic/Services/SourceControl/SubmitToolPerforce.h b/Engine/Source/Programs/SubmitTool/Private/Logic/Services/SourceControl/SubmitToolPerforce.h
--- a/Engine/Source/Programs/SubmitTool/Private/Logic/Services/SourceControl/SubmitToolPerforce.h
+++ b/Engine/Source/Programs/SubmitTool/Private/Logic/Services/SourceControl/SubmitToolPerforce.h
@@ -81,7 +81,7 @@ public:
 	virtual void GetGroups(const FOnGroupsGet::FDelegate& Callback) override;
 	virtual void GetUsersAndGroups(const FOnUsersAndGroupsGet::FDelegate& Callback) override;
 
-	virtual UE::Tasks::TTask<bool> DownloadFiles(TArray<FString>&& FilePaths, TArray<FSharedBuffer>& OutFileBuffers) override;
+	virtual UE::Tasks::TTask<bool> DownloadFiles(TArray<FString>&& FilePaths, TMap<FString, FSharedBuffer>& OutFileBuffers) override;
 
 	virtual bool IsAvailable() const override
 	{
diff --git a/Engine/Source/Programs/SubmitTool/Private/Logic/Validators/CoreRedirectsValidator.cpp b/Engine/Source/Programs/SubmitTool/Private/Logic/Validators/CoreRedirectsValidator.cpp
--- a/Engine/Source/Programs/SubmitTool/Private/Logic/Validators/CoreRedirectsValidator.cpp
+++ b/Engine/Source/Programs/SubmitTool/Private/Logic/Validators/CoreRedirectsValidator.cpp
@@ -106,7 +106,7 @@ bool FCoreRedirectsValidator::DoWork(const UE::Tasks::FCancellationToken& InCanc
 	}
 
 	// download remaining files
-	TArray<FSharedBuffer> FileBuffers;
+	TMap<FString, FSharedBuffer> FileBuffers;
 	if (!ServiceProvider.Pin()->GetService<ISTSourceControlService>()->DownloadFiles(CopyTemp(FilePaths), FileBuffers).GetResult() || FileBuffers.Num() != FilePaths.Num())
 	{
 		LogFailure(FString::Printf(TEXT("[%s]: Unable to perform validation because we could not download %d unedited file(s) from version control."), *GetValidatorName(), FilePaths.Num()));
@@ -115,15 +115,29 @@ bool FCoreRedirectsValidator::DoWork(const UE::Tasks::FCancellationToken& InCanc
 
 	// validate downloaded files
 	TRACE_CPUPROFILER_EVENT_SCOPE(ValidateDownloadedFiles);
-	for (int32 I = 0, N = FilePaths.Num(); I < N; ++I)
+	for (const TPair<FString, FSharedBuffer>& DownloadedFile : FileBuffers)
 	{
 		if (InCancellationToken.IsCanceled())
 		{
 			return false;
 		}
 
-		const FString& Filename = FilePaths[I];
-		const FSharedBuffer& FileBuffer = FileBuffers[I];
+		FString Filename;
+		for (int32 i = 0; i < FilePaths.Num(); ++i)
+		{
+			if (FPaths::GetCleanFilename(FilePaths[i]).Equals(FPaths::GetCleanFilename(DownloadedFile.Key), ESearchCase::IgnoreCase))
+			{
+				Filename = FilePaths[i];
+				break;
+			}
+		}
+
+		if (Filename.IsEmpty())
+		{
+			continue;
+		}
+
+		const FSharedBuffer& FileBuffer = DownloadedFile.Value;
 
 		if (FileBuffer.GetSize() == 0)
 		{
