Quickfix for callback not being executed when there was a partial download of the files

#Synced-CL 47420659

diff --git a/Engine/Source/Programs/SubmitTool/Private/Logic/Services/SourceControl/SubmitToolPerforce.cpp b/Engine/Source/Programs/SubmitTool/Private/Logic/Services/SourceControl/SubmitToolPerforce.cpp
--- a/Engine/Source/Programs/SubmitTool/Private/Logic/Services/SourceControl/SubmitToolPerforce.cpp
+++ b/Engine/Source/Programs/SubmitTool/Private/Logic/Services/SourceControl/SubmitToolPerforce.cpp
@@ -747,7 +747,7 @@ UE::Tasks::TTask<bool> FSubmitToolPerforce::DownloadFiles(TArray<FString>&& File
 	Command->bCollectData = true;
 	return UE::Tasks::Launch(UE_SOURCE_LOCATION, [this, Command, &OutFileBuffers]
 	{
-		if (RunP4Command(*Command, DefaultFlags | EP4ClientUserFlags::CollectData))
+		if (RunP4Command(*Command, DefaultFlags | EP4ClientUserFlags::CollectData) || Command->Data.Num() != 0)
 		{
 			OutFileBuffers = MoveTemp(Command->Data);
 			return true;
