SubmitTool: Always write a default SubmitTool.utrace to the log dir with the cpu,log,bookmark,region and tasks channels enabled.
Include the default SubmitTool.utrace when clicking CopyAllLogs.
The default trace file is only written when the user has not provided any other trace command line paramters.
In CopyAllLogsClicked, check that each file exists before adding it to the windows clipboard to prevent errors when pasting the data.

#rnx

#Synced-CL 0

diff --git a/Engine/Source/Programs/SubmitTool/Private/CommandLine/CmdLineParameters.cpp b/Engine/Source/Programs/SubmitTool/Private/CommandLine/CmdLineParameters.cpp
--- a/Engine/Source/Programs/SubmitTool/Private/CommandLine/CmdLineParameters.cpp
+++ b/Engine/Source/Programs/SubmitTool/Private/CommandLine/CmdLineParameters.cpp
@@ -2,6 +2,7 @@
 
 #include "CmdLineParameters.h"
 #include "Logging/SubmitToolLog.h"
+#include "HAL/FileManager.h"
 #include "Misc/CommandLine.h"
 
 const FCmdLineParameters FCmdLineParameters::Instance = FCmdLineParameters();
@@ -11,6 +12,23 @@ FCmdLineParameters::FCmdLineParameters()
 {
 }
 
+void FCmdLineParameters::SetupTraceParameters() const
+{
+	// Setup tracing to file if no existing trace options are found on the command line
+	FString CommandLine = FCommandLine::Get();
+	if (CommandLine.Find("-trace") == INDEX_NONE)
+	{
+		FString TraceFile = FPaths::Combine(FPaths::ProjectLogDir(), "SubmitTool.utrace");
+		TraceFile = FPaths::ConvertRelativePathToFull(TraceFile);
+		if (IFileManager::Get().FileExists(*TraceFile))
+		{
+			IFileManager::Get().Delete(*TraceFile);
+		}
+		FString TraceArgs = FString::Printf(TEXT("-trace=cpu,log,bookmark,region,tasks -tracefile=\"%s\""), *TraceFile);
+		FCommandLine::Append(*TraceArgs);
+	}
+}
+
 bool FCmdLineParameters::ValidateParameters() const
 {
 	const TCHAR* CommandLine = FCommandLine::Get();
diff --git a/Engine/Source/Programs/SubmitTool/Private/CommandLine/CmdLineParameters.h b/Engine/Source/Programs/SubmitTool/Private/CommandLine/CmdLineParameters.h
--- a/Engine/Source/Programs/SubmitTool/Private/CommandLine/CmdLineParameters.h
+++ b/Engine/Source/Programs/SubmitTool/Private/CommandLine/CmdLineParameters.h
@@ -67,6 +67,7 @@ class FCmdLineParameters
 {
 public:
 	FCmdLineParameters();
+	void SetupTraceParameters() const;
 	bool ValidateParameters() const;
 	void LogParameters() const;
 
@@ -78,4 +79,4 @@ private:
 	const static FCmdLineParameters Instance;
 
 	const TArray<TSharedPtr<FCmdLineParameter>> Parameters;
-};
\ No newline at end of file
+};
diff --git a/Engine/Source/Programs/SubmitTool/Private/SubmitToolApp.cpp b/Engine/Source/Programs/SubmitTool/Private/SubmitToolApp.cpp
--- a/Engine/Source/Programs/SubmitTool/Private/SubmitToolApp.cpp
+++ b/Engine/Source/Programs/SubmitTool/Private/SubmitToolApp.cpp
@@ -175,6 +175,11 @@ int RunSubmitTool(const TCHAR* CommandLine)
 	// need to make sure the cwd is correct before doing anything else
 	FPlatformProcess::SetCurrentWorkingDirectoryToBaseDir();
 
+	// setup command line
+	const FCmdLineParameters& CmdLineParameters = FCmdLineParameters::Get();
+	CmdLineParameters.SetupTraceParameters();
+	CommandLine = FCommandLine::Get();
+
 	// start up the main loop
 	GEngineLoop.PreInit(CommandLine);
 
@@ -186,10 +191,10 @@ int RunSubmitTool(const TCHAR* CommandLine)
 
 	UE_LOG(LogSubmitToolDebug, Log, TEXT("%s"), CommandLine);
 
-	if(!FCmdLineParameters::Get().ValidateParameters())
+	if(!CmdLineParameters.ValidateParameters())
 	{
 		UE_LOG(LogSubmitTool, Error, TEXT("Command line is not valid"));
-		FCmdLineParameters::Get().LogParameters();
+		CmdLineParameters.LogParameters();
 		FModelInterface::SetErrorState();
 	}
 
@@ -212,7 +217,7 @@ int RunSubmitTool(const TCHAR* CommandLine)
 		TUniquePtr<FSubmitToolUserPrefs> UserPrefs = FSubmitToolUserPrefs::Initialize(GetUserPrefsPath());
 
 		FString ParameterFile;
-		FCmdLineParameters::Get().GetValue(FSubmitToolCmdLine::ParameterFile, ParameterFile);
+		CmdLineParameters.GetValue(FSubmitToolCmdLine::ParameterFile, ParameterFile);
 		FSubmitToolParametersBuilder ParametersBuilder(ParameterFile);
 		FSubmitToolParameters Parameters = ParametersBuilder.Build();
 
diff --git a/Engine/Source/Programs/SubmitTool/Private/View/Widgets/SubmitToolWidget.cpp b/Engine/Source/Programs/SubmitTool/Private/View/Widgets/SubmitToolWidget.cpp
--- a/Engine/Source/Programs/SubmitTool/Private/View/Widgets/SubmitToolWidget.cpp
+++ b/Engine/Source/Programs/SubmitTool/Private/View/Widgets/SubmitToolWidget.cpp
@@ -34,6 +34,7 @@
 #include "Version/AppVersion.h"
 
 #include "Async/Async.h"
+#include "HAL/FileManager.h"
 #include "Modules/ModuleManager.h"
 #include "Framework/Docking/TabManager.h"
 #include "Framework/Application/SlateApplication.h"
@@ -832,6 +833,8 @@ FReply SubmitToolWidget::SubmitClicked()
 
 FReply SubmitToolWidget::ValidateClicked()
 {
+	TRACE_BOOKMARK(TEXT("ValidateClicked"));
+
 	if(ModelInterface->IsValidationRunning())
 	{
 		ModelInterface->CancelValidations();
@@ -935,10 +938,17 @@ void SubmitToolWidget::OnValidationUpdated(bool bValid)
 #if PLATFORM_WINDOWS
 FReply SubmitToolWidget::CopyAllLogsClicked()
 {
+	TRACE_BOOKMARK(TEXT("CopyAllLogsClicked"));
+	TRACE_CPUPROFILER_EVENT_FLUSH();
+
 	TArray<FString> Files;
 	for (const FString& Path : ModelInterface->GetParameters().CopyLogParameters.LogsToCollect)
 	{
-		Files.Emplace(FPaths::ConvertRelativePathToFull(FConfiguration::Substitute(Path)));
+		FString FullPath = FPaths::ConvertRelativePathToFull(FConfiguration::Substitute(Path));
+		if (IFileManager::Get().FileExists(*FullPath))
+		{
+			Files.Emplace(MoveTemp(FullPath));
+		}
 	}
 
 	FSubmitToolUtils::CopyDiagnosticFilesToClipboard(Files);
