SubmitTool: Target and build file tweaks for faster compile times and minor performance improvements.

#rnx

#Synced-CL 0

diff --git a/Engine/Source/Programs/SubmitTool/Private/CommandLine/CmdLineParameters.cpp b/Engine/Source/Programs/SubmitTool/Private/CommandLine/CmdLineParameters.cpp
--- a/Engine/Source/Programs/SubmitTool/Private/CommandLine/CmdLineParameters.cpp
+++ b/Engine/Source/Programs/SubmitTool/Private/CommandLine/CmdLineParameters.cpp
@@ -12,7 +12,14 @@ FCmdLineParameters::FCmdLineParameters()
 {
 }
 
-void FCmdLineParameters::SetupTraceParameters() const
+const TCHAR* FCmdLineParameters::InitializeParameters() const
+{
+	ConfigureCoreLimit();
+	ConfigureTrace();
+	return FCommandLine::Get();
+}
+
+void FCmdLineParameters::ConfigureTrace() const
 {
 	// Setup tracing to file if no existing trace options are found on the command line
 	FString CommandLine = FCommandLine::Get();
@@ -29,6 +36,16 @@ void FCmdLineParameters::SetupTraceParameters() const
 	}
 }
 
+void FCmdLineParameters::ConfigureCoreLimit() const
+{
+	// Setup corelimit if an existing option is not found on the command line
+	FString CommandLine = FCommandLine::Get();
+	if (CommandLine.Find("-corelimit") == INDEX_NONE)
+	{
+		FCommandLine::Append(TEXT("-corelimit=16"));
+	}
+}
+
 bool FCmdLineParameters::ValidateParameters() const
 {
 	const TCHAR* CommandLine = FCommandLine::Get();
diff --git a/Engine/Source/Programs/SubmitTool/Private/CommandLine/CmdLineParameters.h b/Engine/Source/Programs/SubmitTool/Private/CommandLine/CmdLineParameters.h
--- a/Engine/Source/Programs/SubmitTool/Private/CommandLine/CmdLineParameters.h
+++ b/Engine/Source/Programs/SubmitTool/Private/CommandLine/CmdLineParameters.h
@@ -67,7 +67,7 @@ class FCmdLineParameters
 {
 public:
 	FCmdLineParameters();
-	void SetupTraceParameters() const;
+	const TCHAR* InitializeParameters() const;
 	bool ValidateParameters() const;
 	void LogParameters() const;
 
@@ -77,6 +77,8 @@ public:
 	bool GetValue(const FString& InKey, FString& OutValue) const;
 private:
 	const static FCmdLineParameters Instance;
+	void ConfigureTrace() const;
+	void ConfigureCoreLimit() const;
 
 	const TArray<TSharedPtr<FCmdLineParameter>> Parameters;
 };
diff --git a/Engine/Source/Programs/SubmitTool/Private/SubmitToolApp.cpp b/Engine/Source/Programs/SubmitTool/Private/SubmitToolApp.cpp
--- a/Engine/Source/Programs/SubmitTool/Private/SubmitToolApp.cpp
+++ b/Engine/Source/Programs/SubmitTool/Private/SubmitToolApp.cpp
@@ -177,8 +177,7 @@ int RunSubmitTool(const TCHAR* CommandLine)
 
 	// setup command line
 	const FCmdLineParameters& CmdLineParameters = FCmdLineParameters::Get();
-	CmdLineParameters.SetupTraceParameters();
-	CommandLine = FCommandLine::Get();
+	CommandLine = CmdLineParameters.InitializeParameters();
 
 	// start up the main loop
 	GEngineLoop.PreInit(CommandLine);
diff --git a/Engine/Source/Programs/SubmitTool/SubmitTool.Build.cs b/Engine/Source/Programs/SubmitTool/SubmitTool.Build.cs
--- a/Engine/Source/Programs/SubmitTool/SubmitTool.Build.cs
+++ b/Engine/Source/Programs/SubmitTool/SubmitTool.Build.cs
@@ -6,6 +6,8 @@ public class SubmitTool : ModuleRules
 {
 	public SubmitTool(ReadOnlyTargetRules Target) : base(Target)
 	{
+		NumIncludedBytesPerUnityCPPOverride = 64 * 1024;
+	
 		PublicIncludePathModuleNames.AddRange(
 			new string[] {
 				"Launch",
diff --git a/Engine/Source/Programs/SubmitTool/SubmitTool.Target.cs b/Engine/Source/Programs/SubmitTool/SubmitTool.Target.cs
--- a/Engine/Source/Programs/SubmitTool/SubmitTool.Target.cs
+++ b/Engine/Source/Programs/SubmitTool/SubmitTool.Target.cs
@@ -47,6 +47,10 @@ public class SubmitToolTarget : TargetRules
 		GlobalDefinitions.Add("EXCLUDE_NONPAK_UE_EXTENSIONS=0");
 		GlobalDefinitions.Add("WITH_AUTOMATION_WORKER=0");
 		GlobalDefinitions.Add("UE_EXTERNAL_PROFILING_ENABLED=0");
+		GlobalDefinitions.Add("STATS=0");
+		GlobalDefinitions.Add("ENABLE_STATNAMEDEVENTS=0");
+		GlobalDefinitions.Add("ALLOW_HITCH_DETECTION=0");
+		GlobalDefinitions.Add("UE_USE_MALLOC_FILL_BYTES=0");
 
 		OptedInModulePlatforms = new UnrealTargetPlatform[] { UnrealTargetPlatform.Win64, UnrealTargetPlatform.Mac,
 															  UnrealTargetPlatform.Linux, UnrealTargetPlatform.LinuxArm64 };
