SubmitTool: Add missing space to the appended command line parameters.
Log the original and modified command lines.

#rnx

#Synced-CL 0

diff --git a/Engine/Source/Programs/SubmitTool/Private/CommandLine/CmdLineParameters.cpp b/Engine/Source/Programs/SubmitTool/Private/CommandLine/CmdLineParameters.cpp
--- a/Engine/Source/Programs/SubmitTool/Private/CommandLine/CmdLineParameters.cpp
+++ b/Engine/Source/Programs/SubmitTool/Private/CommandLine/CmdLineParameters.cpp
@@ -14,8 +14,10 @@ FCmdLineParameters::FCmdLineParameters()
 
 const TCHAR* FCmdLineParameters::InitializeParameters() const
 {
+	UE_LOG(LogSubmitTool, Log, TEXT("Original command line: %s"), FCommandLine::Get());
 	ConfigureCoreLimit();
 	ConfigureTrace();
+	UE_LOG(LogSubmitTool, Log, TEXT("Modified command line: %s"), FCommandLine::Get());
 	return FCommandLine::Get();
 }
 
@@ -31,7 +33,7 @@ void FCmdLineParameters::ConfigureTrace() const
 		{
 			IFileManager::Get().Delete(*TraceFile);
 		}
-		FString TraceArgs = FString::Printf(TEXT("-trace=cpu,log,bookmark,region,tasks -tracefile=\"%s\""), *TraceFile);
+		FString TraceArgs = FString::Printf(TEXT(" -trace=cpu,log,bookmark,region,tasks -tracefile=\"%s\""), *TraceFile);
 		FCommandLine::Append(*TraceArgs);
 	}
 }
@@ -42,7 +44,7 @@ void FCmdLineParameters::ConfigureCoreLimit() const
 	FString CommandLine = FCommandLine::Get();
 	if (CommandLine.Find("-corelimit") == INDEX_NONE)
 	{
-		FCommandLine::Append(TEXT("-corelimit=16"));
+		FCommandLine::Append(TEXT(" -corelimit=16"));
 	}
 }
 
