SubmitTool: Enable running from cmd line with just -cl as a parameter.

Enable .ushell workflows like these when submitting a cl in the same workspace as the SubmitTool instance.
.run program SubmitTool --build -- -cl nnnnnn
.run program SubmitTool --attach -- -cl nnnnnn

#rnx

#Synced-CL 0

diff --git a/Engine/Source/Programs/SubmitTool/Private/CommandLine/CmdLineParameters.cpp b/Engine/Source/Programs/SubmitTool/Private/CommandLine/CmdLineParameters.cpp
--- a/Engine/Source/Programs/SubmitTool/Private/CommandLine/CmdLineParameters.cpp
+++ b/Engine/Source/Programs/SubmitTool/Private/CommandLine/CmdLineParameters.cpp
@@ -4,6 +4,8 @@
 #include "Logging/SubmitToolLog.h"
 #include "HAL/FileManager.h"
 #include "Misc/CommandLine.h"
+#include "Misc/FileHelper.h"
+#include "String/ParseTokens.h"
 
 const FCmdLineParameters FCmdLineParameters::Instance = FCmdLineParameters();
 
@@ -15,12 +17,128 @@ FCmdLineParameters::FCmdLineParameters()
 const TCHAR* FCmdLineParameters::InitializeParameters() const
 {
 	UE_LOG(LogSubmitTool, Log, TEXT("Original command line: %s"), FCommandLine::Get());
+	ConfigurePerforce();
 	ConfigureCoreLimit();
 	ConfigureTrace();
 	UE_LOG(LogSubmitTool, Log, TEXT("Modified command line: %s"), FCommandLine::Get());
 	return FCommandLine::Get();
 }
 
+void FCmdLineParameters::ConfigureCoreLimit() const
+{
+	// Setup corelimit if an existing option is not found on the command line
+	FString CommandLine = FCommandLine::Get();
+	if (CommandLine.Find("-corelimit") == INDEX_NONE)
+	{
+		FCommandLine::Append(TEXT(" -corelimit=16"));
+	}
+}
+
+void FCmdLineParameters::ConfigurePerforce() const
+{
+	const TCHAR* CommandLine = FCommandLine::Get();
+
+	FString ServerAndPort;
+	FString UserName;
+	FString ClientName;
+
+	// 1. Check if all perforce parameters have been supplied on the command line
+	int32 NumParameters = 0;
+	NumParameters += FParse::Value(CommandLine, *FSubmitToolCmdLine::P4Server, ServerAndPort);
+	NumParameters += FParse::Value(CommandLine, *FSubmitToolCmdLine::P4User, UserName);
+	NumParameters += FParse::Value(CommandLine, *FSubmitToolCmdLine::P4Client, ClientName);
+
+	UE_LOG(LogSubmitTool, Log, TEXT("Perforce parameters after reading command line: Server: '%s' | User: '%s' | Workspace: '%s'"), *ServerAndPort, *UserName, *ClientName);
+	if (NumParameters == 3)
+	{
+		return;
+	}
+
+	// 2. Fill missing perforce parameters from an existing p4config file in the root directory.
+	// The cl to review is expected to be in the same P4Client as the directory where SubmitTool is running.
+	TArray<FString, TInlineAllocator<3>> ConfigFilenames;
+	if (FString Filename = FPlatformMisc::GetEnvironmentVariable(TEXT("P4CONFIG")); !Filename.IsEmpty())
+	{
+		ConfigFilenames.Add(MoveTemp(Filename));
+	}
+	else
+	{
+		ConfigFilenames.Append({"p4config.txt", ".p4config.txt", ".p4config"});
+	}
+	for (const FString& P4ConfigFilename : ConfigFilenames)
+	{
+		FString P4ConfigFilePath = FPaths::RootDir() / P4ConfigFilename;
+		if (IFileManager::Get().FileExists(*P4ConfigFilePath))
+		{
+			FFileHelper::LoadFileToStringWithLineVisitor(
+				*P4ConfigFilePath,
+				[&NumParameters, &ServerAndPort, &UserName, &ClientName](FStringView Line)
+				{
+					using namespace UE::String;
+					TArray<FStringView, TInlineAllocator<2>> Tokens;
+					ParseTokens(Line, '=', Tokens, EParseTokensOptions::Trim | EParseTokensOptions::SkipEmpty);
+					if (Tokens.Num() != 2)
+					{
+						return;
+					}
+					if (Tokens[0] == TEXTVIEW("P4PORT"))
+					{
+						if (ServerAndPort.IsEmpty())
+						{
+							++NumParameters;
+							ServerAndPort = Tokens[1];
+							FCommandLine::Append(*WriteToString<128>(" -", FSubmitToolCmdLine::P4Server, " ", ServerAndPort));
+						}
+					}
+					else if (Tokens[0] == TEXTVIEW("P4USER"))
+					{
+						if (UserName.IsEmpty())
+						{
+							++NumParameters;
+							UserName = Tokens[1];
+							FCommandLine::Append(*WriteToString<128>(" -", FSubmitToolCmdLine::P4User, " ", UserName));
+						}
+					}
+					else if (Tokens[0] == TEXTVIEW("P4CLIENT"))
+					{
+						if (ClientName.IsEmpty())
+						{
+							++NumParameters;
+							ClientName = Tokens[1];
+							FCommandLine::Append(*WriteToString<128>(" -", FSubmitToolCmdLine::P4Client, " ", ClientName));
+						}
+					}
+				});
+			UE_LOG(LogSubmitTool, Log, TEXT("Perforce parameters after reading file \"%s\": Server: '%s' | User: '%s' | Workspace: '%s'"),
+				*P4ConfigFilePath, *ServerAndPort, *UserName, *ClientName);
+			if (NumParameters == 3)
+			{
+				return;
+			}
+		}
+	}
+
+	// 3. Fill missing perforce parameters from P4 environment variables.
+	if (ServerAndPort.IsEmpty())
+	{
+		ServerAndPort = FPlatformMisc::GetEnvironmentVariable(TEXT("P4PORT"));
+		FCommandLine::Append(*WriteToString<128>(" -", FSubmitToolCmdLine::P4Server, " ", ServerAndPort));
+	}
+	if (UserName.IsEmpty())
+	{
+		UserName = FPlatformMisc::GetEnvironmentVariable(TEXT("P4USER"));
+		FCommandLine::Append(*WriteToString<128>(" -", FSubmitToolCmdLine::P4User, " ", UserName));
+	}
+	if (ClientName.IsEmpty())
+	{
+		ClientName = FPlatformMisc::GetEnvironmentVariable(TEXT("P4CLIENT"));;
+		FCommandLine::Append(*WriteToString<128>(" -", FSubmitToolCmdLine::P4Client, " ", ClientName));
+	}
+	UE_LOG(LogSubmitTool, Log, TEXT("Perforce parameters after reading environment variables: Server: '%s' | User: '%s' | Workspace: '%s'"),
+		*ServerAndPort, *UserName, *ClientName);
+}
+
+
 void FCmdLineParameters::ConfigureTrace() const
 {
 	// Setup tracing to file if no existing trace options are found on the command line
@@ -33,18 +151,7 @@ void FCmdLineParameters::ConfigureTrace() const
 		{
 			IFileManager::Get().Delete(*TraceFile);
 		}
-		FString TraceArgs = FString::Printf(TEXT(" -trace=cpu,log,bookmark,region,tasks -tracefile=\"%s\""), *TraceFile);
-		FCommandLine::Append(*TraceArgs);
-	}
-}
-
-void FCmdLineParameters::ConfigureCoreLimit() const
-{
-	// Setup corelimit if an existing option is not found on the command line
-	FString CommandLine = FCommandLine::Get();
-	if (CommandLine.Find("-corelimit") == INDEX_NONE)
-	{
-		FCommandLine::Append(TEXT(" -corelimit=16"));
+		FCommandLine::Append(*WriteToString<128>(" -trace=cpu,log,bookmark,region,tasks -tracefile=\"", *TraceFile, "\""));
 	}
 }
 
diff --git a/Engine/Source/Programs/SubmitTool/Private/CommandLine/CmdLineParameters.h b/Engine/Source/Programs/SubmitTool/Private/CommandLine/CmdLineParameters.h
--- a/Engine/Source/Programs/SubmitTool/Private/CommandLine/CmdLineParameters.h
+++ b/Engine/Source/Programs/SubmitTool/Private/CommandLine/CmdLineParameters.h
@@ -77,8 +77,9 @@ public:
 	bool GetValue(const FString& InKey, FString& OutValue) const;
 private:
 	const static FCmdLineParameters Instance;
-	void ConfigureTrace() const;
 	void ConfigureCoreLimit() const;
+	void ConfigurePerforce() const;
+	void ConfigureTrace() const;
 
 	const TArray<TSharedPtr<FCmdLineParameter>> Parameters;
 };
diff --git a/Engine/Source/Programs/SubmitTool/Private/Logic/Services/SourceControl/SubmitToolPerforce.cpp b/Engine/Source/Programs/SubmitTool/Private/Logic/Services/SourceControl/SubmitToolPerforce.cpp
--- a/Engine/Source/Programs/SubmitTool/Private/Logic/Services/SourceControl/SubmitToolPerforce.cpp
+++ b/Engine/Source/Programs/SubmitTool/Private/Logic/Services/SourceControl/SubmitToolPerforce.cpp
@@ -134,7 +134,7 @@ FSubmitToolPerforce::FSubmitToolPerforce(const FSubmitToolParameters& InParamete
 
 	if(SCCProvider.IsValid())
 	{
-		UE_LOG(LogSubmitTool, Log, TEXT("Setting Perforce Connection parameters: %s | User: %s | Workspace: %s"), *PerforceServerAndPort, *PerforceUserName, *PerforceClientName);
+		UE_LOG(LogSubmitTool, Log, TEXT("Setting Perforce Connection parameters: '%s' | User: '%s' | Workspace: '%s'"), *PerforceServerAndPort, *PerforceUserName, *PerforceClientName);
 
 		ISourceControlProvider::FInitResult ConnectionResult = SCCProvider->Init(ISourceControlProvider::EInitFlags::None);
 		if(!ConnectionResult.bIsAvailable)
