SubmitTool: Parse (and wait for) the initial get stream task from GetRootStreamName and GetCurrentStreamName instead of doing it in a GT task.
The current stream name was being requested by the UI on startup before the first tick, which means the GT task had not been executed yet, and a duplicated query was performed.

#rnx

#Synced-CL 0

diff --git a/Engine/Source/Programs/SubmitTool/Private/Logic/Services/SourceControl/SubmitToolPerforce.cpp b/Engine/Source/Programs/SubmitTool/Private/Logic/Services/SourceControl/SubmitToolPerforce.cpp
--- a/Engine/Source/Programs/SubmitTool/Private/Logic/Services/SourceControl/SubmitToolPerforce.cpp
+++ b/Engine/Source/Programs/SubmitTool/Private/Logic/Services/SourceControl/SubmitToolPerforce.cpp
@@ -152,12 +152,7 @@ FSubmitToolPerforce::FSubmitToolPerforce(const FSubmitToolParameters& InParamete
 	}
 
 	TickHandle = FTSTicker::GetCoreTicker().AddTicker(FTickerDelegate::CreateRaw(this, &FSubmitToolPerforce::Tick));
-	UE::Tasks::TTask<TSharedRef<FSCCommand>> GetStreamTask = GetStream();
-
-	UE::Tasks::Launch(UE_SOURCE_LOCATION, [this, GetStreamTask] () mutable
-	{
-		ParseStreamCommand(*GetStreamTask.GetResult());
-	}, GetStreamTask, LowLevelTasks::ETaskPriority::Normal, UE::Tasks::EExtendedTaskPriority::GameThreadNormalPri);
+	InitialStreamTask = GetStream();
 }
 
 FSubmitToolPerforce::~FSubmitToolPerforce()
@@ -801,9 +796,20 @@ const FSCCStream* FSubmitToolPerforce::GetSCCStream(const FString& InStreamName)
 	return nullptr;
 }
 
+void FSubmitToolPerforce::WaitForInitialStreamTask()
+{
+	if (InitialStreamTask.IsValid())
+	{
+		TRACE_CPUPROFILER_EVENT_SCOPE(FSubmitToolPerforce::WaitForInitialStreamTask);
+		ParseStreamCommand(*InitialStreamTask.GetResult());
+		InitialStreamTask = {};
+	}
+}
+
 const FString FSubmitToolPerforce::GetRootStreamName()
 {
 	TRACE_CPUPROFILER_EVENT_SCOPE(FSubmitToolPerforce::GetRootStreamName);
+	WaitForInitialStreamTask();
 	for (int32 I = 0; I < 2; ++I)
 	{
 		if (StreamHierarchy.Num() > 0)
@@ -822,6 +828,7 @@ const FString FSubmitToolPerforce::GetRootStreamName()
 const FString FSubmitToolPerforce::GetCurrentStreamName()
 {
 	TRACE_CPUPROFILER_EVENT_SCOPE(FSubmitToolPerforce::GetCurrentStreamName);
+	WaitForInitialStreamTask();
 	for (int32 I = 0; I < 2; ++I)
 	{
 		if (StreamHierarchy.Num() > 0)
diff --git a/Engine/Source/Programs/SubmitTool/Private/Logic/Services/SourceControl/SubmitToolPerforce.h b/Engine/Source/Programs/SubmitTool/Private/Logic/Services/SourceControl/SubmitToolPerforce.h
--- a/Engine/Source/Programs/SubmitTool/Private/Logic/Services/SourceControl/SubmitToolPerforce.h
+++ b/Engine/Source/Programs/SubmitTool/Private/Logic/Services/SourceControl/SubmitToolPerforce.h
@@ -111,12 +111,15 @@ protected:
 private:
 	bool RunP4Command(FSCCommand& Command, EP4ClientUserFlags CmdFlags = DefaultFlags);
 
-	const FSubmitToolParameters& Parameters;
 	UE::Tasks::TTask<TSharedRef<FSCCommand>> GetStream(const FString& InStream = FString());
 	void ParseStreamCommand(const FSCCommand& InCompletedCommand, const FString& InRequestedStream = FString(), bool bParseHierarchy = true);
+	void WaitForInitialStreamTask();
+
+	const FSubmitToolParameters& Parameters;
 	TMap<FString, TSharedPtr<FUserData>> CachedUsers;
 	TArray<TSharedPtr<FUserData>> CachedUsersArray;
 	TArray<TSharedPtr<FString>> CachedGroupsArray;
+	UE::Tasks::TTask<TSharedRef<FSCCommand>> InitialStreamTask;
 	bool bUserTaskIsRunning = false;
 	std::atomic<bool> bGroupTaskIsRunning = false;
 	bool bIsUnicodeServer = false;
