Fix for submit tool not noticing a shelf refresh made in p4v when the number of shelved files are the same.
Specifying in the UI when preflight is waiting for shelf

#Synced-CL 48450364

diff --git a/Engine/Source/Programs/SubmitTool/Private/Logic/ChangelistService.cpp b/Engine/Source/Programs/SubmitTool/Private/Logic/ChangelistService.cpp
--- a/Engine/Source/Programs/SubmitTool/Private/Logic/ChangelistService.cpp
+++ b/Engine/Source/Programs/SubmitTool/Private/Logic/ChangelistService.cpp
@@ -311,10 +311,7 @@ void FChangelistService::UpdateCachedChangelistState()
 		PrintFilesAndShelvedFiles();
 	}
 
-	if(ChangeType != ETaskArea::None)
-	{
-		CLRefreshCallback.ExecuteIfBound(ChangeType);
-	}
+	CLRefreshCallback.ExecuteIfBound(ChangeType);
 }
 
 bool FChangelistService::AreCLDescriptionsIdentical() const
@@ -819,14 +816,16 @@ const TMap<FString, FString>& FChangelistService::GetOpenFilesInChangelists()
 
 void FChangelistService::CheckShelfIsUpToDate(TFunction<void(bool, TArray<FString>)> OnComplete)
 {
-	check(OnComplete);
 	TRACE_CPUPROFILER_EVENT_SCOPE(FChangelistService::CheckShelfIsUpToDate);
 
 	TArray<FString> FilePathsOutOfDate;
 
 	if (GetOpenedFilesDepotPaths() != GetShelvedFilesDepotPaths())
 	{
-		OnComplete(false, MoveTemp(FilePathsOutOfDate));
+		if(OnComplete)
+		{
+			OnComplete(false, MoveTemp(FilePathsOutOfDate));
+		}
 		return;
 	}
 	
@@ -853,7 +852,11 @@ void FChangelistService::CheckShelfIsUpToDate(TFunction<void(bool, TArray<FStrin
 			if (FilePathsOutOfDate.Num() == 0)
 			{
 				bIsShelfReady = true;
-				OnComplete(bIsShelfReady, MoveTemp(FilePathsOutOfDate));
+
+				if (OnComplete)
+				{
+					OnComplete(bIsShelfReady, MoveTemp(FilePathsOutOfDate));
+				}
 				return;
 			}
 		}
@@ -907,7 +910,10 @@ void FChangelistService::CheckShelfIsUpToDate(TFunction<void(bool, TArray<FStrin
 					ServiceProvider->GetService<ICacheDataService>()->SetLastShelfKnownUpToDate(CLID, FDateTime::UtcNow());
 				}
 
-				OnComplete(bIsShelfReady, FilesWithChanges);
+				if (OnComplete)
+				{
+					OnComplete(bIsShelfReady, FilesWithChanges);
+				}
 				bDiffTaskRunning = false;
 			}));
 	}
@@ -941,7 +947,7 @@ void FChangelistService::EnsureShelfIsCurrent()
 					DialogResult = FDialogFactory::ShowDialog(
 						TextTitle, 
 						TextDescription,
-						TArray<FString>{ TEXT("Re-shelve files"), TEXT("Cancel")},
+						TArray<FString>{ TEXT("Re-shelve files"), TEXT("I will reshelve in P4V")},
 						FSubmitToolUtils::BuildUserPrefCheckboxUI(FSubmitToolUserPrefs::Get()->bAutoStompShelf, FText::FromString(TEXT("Always stomp a shelf when it doesn't match local state, don't ask again"))));
 				}
 
diff --git a/Engine/Source/Programs/SubmitTool/Private/Models/ModelInterface.cpp b/Engine/Source/Programs/SubmitTool/Private/Models/ModelInterface.cpp
--- a/Engine/Source/Programs/SubmitTool/Private/Models/ModelInterface.cpp
+++ b/Engine/Source/Programs/SubmitTool/Private/Models/ModelInterface.cpp
@@ -902,6 +902,35 @@ void FModelInterface::OnChangelistReady(bool bIsValid)
 	}
 }
 
+void FModelInterface::OnChangelistRefresh(ETaskArea InChangeType)
+{
+	// If we're waiting for a shelf, check again here to detect same file number but different content when files were refreshed in P4V
+	if (!ChangelistService->IsShelfReady())
+	{
+		ChangelistService->CheckShelfIsUpToDate(nullptr);
+	}
+
+	if (InChangeType != ETaskArea::None)
+	{
+		if ((InChangeType & ETaskArea::Changelist) == ETaskArea::Changelist)
+		{
+			TagService->ParseCLDescription();
+		}
+
+		if ((InChangeType & ETaskArea::LocalFiles) == ETaskArea::LocalFiles)
+		{
+			FileRefreshedCallback.Broadcast();
+		}
+
+		if ((InChangeType & (ETaskArea::LocalFiles | ETaskArea::ShelvedFiles)) != ETaskArea::None)
+		{
+			RefreshStateBasedOnFiles();
+		}
+
+		ValidationService->InvalidateForChanges(InChangeType);
+	}
+}
+
 
 void FModelInterface::RevertUnchangedAndSubmit()
 {
diff --git a/Engine/Source/Programs/SubmitTool/Private/Models/ModelInterface.h b/Engine/Source/Programs/SubmitTool/Private/Models/ModelInterface.h
--- a/Engine/Source/Programs/SubmitTool/Private/Models/ModelInterface.h
+++ b/Engine/Source/Programs/SubmitTool/Private/Models/ModelInterface.h
@@ -179,6 +179,7 @@ public:
 
 	void RequestPreflight(bool bForceStart = false);
 	bool IsPreflightRequestInProgress() const										{ return PreflightService->IsRequestInProgress(); }
+	bool IsShelfReady() const														{ return ChangelistService->IsShelfReady(); }
 	bool IsPreflightQueued() const													{ return bPreflightQueued; }
 	void RefreshPreflightInformation() const										{ PreflightService->FetchPreflightInfo(); UE_LOG(LogSubmitTool, Log, TEXT("Requesting preflight information...")); }
 	const TUniquePtr<FPreflightList>& GetPreflightData()							{ return PreflightService->GetPreflightData(); }
@@ -246,26 +247,7 @@ public:
 
 private:	
 	const FSubmitToolParameters& Parameters;
-	void OnChangelistRefresh(ETaskArea InChangeType)
-	{
-		if((InChangeType & ETaskArea::Changelist) == ETaskArea::Changelist)
-		{
-			TagService->ParseCLDescription();
-		}
-
-		if((InChangeType & ETaskArea::LocalFiles) == ETaskArea::LocalFiles)
-		{
-			FileRefreshedCallback.Broadcast();		
-		}
-
-		if((InChangeType & (ETaskArea::LocalFiles | ETaskArea::ShelvedFiles)) != ETaskArea::None)
-		{
-			RefreshStateBasedOnFiles();
-		}
-
-
-		ValidationService->InvalidateForChanges(InChangeType);
-	}
+	void OnChangelistRefresh(ETaskArea InChangeType);
 	void PreValidate(const TArray<const TWeakPtr<const FValidatorBase>>& InValidators) const;
 	void RefreshStateBasedOnFiles();
 	void OnChangelistReady(bool bIsValid);
diff --git a/Engine/Source/Programs/SubmitTool/Private/View/Widgets/TagWidget.cpp b/Engine/Source/Programs/SubmitTool/Private/View/Widgets/TagWidget.cpp
--- a/Engine/Source/Programs/SubmitTool/Private/View/Widgets/TagWidget.cpp
+++ b/Engine/Source/Programs/SubmitTool/Private/View/Widgets/TagWidget.cpp
@@ -561,7 +561,14 @@ void STagWidget::Construct(const FArguments& InArgs)
 							{
 								if (ModelInterface->IsPreflightRequestInProgress())
 								{
-									return FText::FromString(TEXT("Requesting..."));
+									if(ModelInterface->IsShelfReady())
+									{ 
+										return FText::FromString(TEXT("Requesting..."));
+									}
+									else
+									{ 
+										return FText::FromString(TEXT("Waiting for Perforce"));
+									}
 								}
 								else if (ModelInterface->IsPreflightQueued())
 								{
