[Submit Tool] Adding submit tool ini for engine restricted folder in the config layer

#Synced-CL 48671645

diff --git a/Engine/Source/Programs/SubmitTool/Private/Parameters/SubmitToolParametersBuilder.cpp b/Engine/Source/Programs/SubmitTool/Private/Parameters/SubmitToolParametersBuilder.cpp
--- a/Engine/Source/Programs/SubmitTool/Private/Parameters/SubmitToolParametersBuilder.cpp
+++ b/Engine/Source/Programs/SubmitTool/Private/Parameters/SubmitToolParametersBuilder.cpp
@@ -32,6 +32,12 @@ FSubmitToolParametersBuilder::FSubmitToolParametersBuilder(const FString& InPara
 
 	if(!RootDir.IsEmpty())
 	{
+		FString* EngineBase = new FString(FPaths::Combine(RootDir, TEXT("/Engine/Restricted/NotForLicensees/Config/{TYPE}.ini")));
+		ConfigHierarchy.Add({ TEXT("EngineBase"), **EngineBase, EConfigLayerFlags::NoExpand });
+
+		FString* EnginePlatform = new FString(FPaths::Combine(RootDir, TEXT("Engine/Restricted/NotForLicensees/Config/{PLATFORM}/{PLATFORM}{TYPE}.ini")));
+		ConfigHierarchy.Add({ TEXT("EnginePlatform"), **EnginePlatform });
+
 		IFileManager::Get().IterateDirectory(*FConfiguration::Substitute(TEXT("$(root)")),
 			[this](const TCHAR* FileOrDir, bool bIsDir) -> bool
 			{
