Fixing allowlist status when having more than one hardlocked section active and not being allowlisted in all of them
Removed SubmitLocked state in favour of a bool since with allowlist exceptions, we were vulnerable to losing the state when going to submit and failing.

#Synced-CL 48728512

diff --git a/Engine/Source/Developer/SubmitToolCore/Public/SubmitToolCoreUtils.h b/Engine/Source/Developer/SubmitToolCore/Public/SubmitToolCoreUtils.h
--- a/Engine/Source/Developer/SubmitToolCore/Public/SubmitToolCoreUtils.h
+++ b/Engine/Source/Developer/SubmitToolCore/Public/SubmitToolCoreUtils.h
@@ -15,6 +15,14 @@ public:
 
 	static bool IsFileInHierarchy(const FString& InWildcard, const FString& InPath);
 
+	template <int32 BufferSize = 256, typename RangeType>
+	static FString StringBuilderJoin(const RangeType& InRange, const FString&& Delimiter)
+	{
+		TStringBuilder<BufferSize> Builder;
+		Builder.Join(InRange, Delimiter);
+		return Builder.ToString();
+	}
+
 private:
 	static TMap<FString, TMap<bool, TSet<FString>>> HierarchyWildcardsCache;
 };
diff --git a/Engine/Source/Programs/SubmitTool/Private/Logic/P4LockdownService.cpp b/Engine/Source/Programs/SubmitTool/Private/Logic/P4LockdownService.cpp
--- a/Engine/Source/Programs/SubmitTool/Private/Logic/P4LockdownService.cpp
+++ b/Engine/Source/Programs/SubmitTool/Private/Logic/P4LockdownService.cpp
@@ -5,6 +5,7 @@
 #include "Logic/Services/Interfaces/ISTSourceControlService.h"
 #include "Logic/Services/SubmitToolServiceProvider.h"
 #include "Logging/SubmitToolLog.h"
+#include "SubmitToolCoreUtils.h"
 #include "Misc/ConfigCacheIni.h"
 #include "Misc/FileHelper.h"
 #include "Misc/Paths.h"
@@ -91,7 +92,7 @@ FPathLockdownResult FP4LockdownService::ArePathsInLockdown(const TArray<FString>
 {	
 	TRACE_CPUPROFILER_EVENT_SCOPE(FP4LockdownService::ArePathsInLockdown);
 
-	FPathLockdownResult GlobalResult{ FString(), FString(), false, true };
+	FPathLockdownResult GlobalResult{ TSet<FString>(), TSet<FString>(), false, true };
 
 	if (InPaths.IsEmpty())
 	{
@@ -112,8 +113,6 @@ FPathLockdownResult FP4LockdownService::ArePathsInLockdown(const TArray<FString>
 		AllowListState = bShouldParseDownloadedData ? EAllowListState::Downloaded : EAllowListState::Cached;
 	}
 
-	TSet<FString> UniqueGroups;
-	TSet<FString> UniqueMessages;
 	// Finally validate the paths
 	for(const FString& Path : InPaths)
 	{
@@ -121,19 +120,15 @@ FPathLockdownResult FP4LockdownService::ArePathsInLockdown(const TArray<FString>
 
 		if(Result.bIsLocked)
 		{
-			UE_LOG(LogSubmitToolP4Debug, Warning, TEXT("File %s is locked down by %s.\n%s"), *Path, *Result.GroupName, *Result.Message);
+			UE_LOG(LogSubmitToolP4Debug, Warning, TEXT("File %s is locked down by %s.\n%s"), *Path, *FSubmitToolCoreUtils::StringBuilderJoin<64>(Result.GroupNames, TEXT(", ")), *FSubmitToolCoreUtils::StringBuilderJoin(Result.Messages, TEXT("\n")));
 			GlobalResult.bAllowlisted &= Result.bAllowlisted;
 			GlobalResult.bIsLocked = Result.bIsLocked;
 
-			UniqueGroups.Add(Result.GroupName);
-			UniqueMessages.Add(Result.Message);
+			GlobalResult.GroupNames.Append(MoveTemp(Result.GroupNames));
+			GlobalResult.Messages.Append(MoveTemp(Result.Messages));
 		}
 	}
 
-	GlobalResult.GroupName = FString::Join(UniqueGroups, TEXT(" - "));
-	GlobalResult.Message = FString::Join(UniqueMessages, TEXT("\n"));
-
-
 	return GlobalResult;
 }
 
@@ -142,7 +137,7 @@ FPathLockdownResult FP4LockdownService::IsPathInLockdown(const FString& InPath)
 	FString PerforceUserName;
 	FCmdLineParameters::Get().GetValue(FSubmitToolCmdLine::P4User, PerforceUserName);
 
-	FPathLockdownResult bOverallLockdownResult { FString(), FString(), false, true};
+	FPathLockdownResult bOverallLockdownResult { TSet<FString>(), TSet<FString>(), false, true};
 
 	for(const FAllowListData& data : AllowListData)
 	{
@@ -178,11 +173,12 @@ FPathLockdownResult FP4LockdownService::IsPathInLockdown(const FString& InPath)
 
 		if (bIsLocked)
 		{
-			bOverallLockdownResult.GroupName = data.GroupName;
-			bOverallLockdownResult.Message = data.Message;
-			bOverallLockdownResult.bAllowlisted = bIsInOverrideAllowlist || data.AllowListers.Contains(PerforceUserName);
+			const bool bAllowlisted = bIsInOverrideAllowlist || data.AllowListers.Contains(PerforceUserName);
+			const FString Message = bAllowlisted ? FString::Printf(TEXT("[Allowlisted] %s"), *data.Message) : data.Message;
+			bOverallLockdownResult.GroupNames.Add(data.GroupName);
+			bOverallLockdownResult.Messages.Add(Message);
 			bOverallLockdownResult.bIsLocked = true;
-			break;
+			bOverallLockdownResult.bAllowlisted &= bAllowlisted;
 		}
 	}
 
diff --git a/Engine/Source/Programs/SubmitTool/Private/Logic/P4LockdownService.h b/Engine/Source/Programs/SubmitTool/Private/Logic/P4LockdownService.h
--- a/Engine/Source/Programs/SubmitTool/Private/Logic/P4LockdownService.h
+++ b/Engine/Source/Programs/SubmitTool/Private/Logic/P4LockdownService.h
@@ -16,8 +16,8 @@ class FSubmitToolServiceProvider;
 
 struct FPathLockdownResult
 {
-	FString GroupName;
-	FString Message;
+	TSet<FString> GroupNames;
+	TSet<FString> Messages;
 	bool bIsLocked;
 	bool bAllowlisted;
 };
diff --git a/Engine/Source/Programs/SubmitTool/Private/Models/ModelInterface.cpp b/Engine/Source/Programs/SubmitTool/Private/Models/ModelInterface.cpp
--- a/Engine/Source/Programs/SubmitTool/Private/Models/ModelInterface.cpp
+++ b/Engine/Source/Programs/SubmitTool/Private/Models/ModelInterface.cpp
@@ -33,6 +33,7 @@
 #include "Logic/PreSubmitOperations/VirtualizationOperation.h"
 #include "Logic/Services/CacheDataService.h"
 #include "Logic/Services/SourceControl/SubmitToolPerforce.h"
+#include "SubmitToolCoreUtils.h"
 
 ESubmitToolAppState FModelInterface::SubmitToolState = ESubmitToolAppState::Initializing;
 FOnStateChanged FModelInterface::OnStateChanged;
@@ -493,7 +494,7 @@ bool FModelInterface::HasSubmitToolTag() const
 
 void FModelInterface::UpdateCLFromP4Async() const
 {
-	if (GetInputEnabled() || SubmitToolState == ESubmitToolAppState::Errored || SubmitToolState == ESubmitToolAppState::SubmitLocked)
+	if (GetInputEnabled() || SubmitToolState == ESubmitToolAppState::Errored)
 	{
 		ChangelistService->FetchChangelistDataAsync();
 	}
@@ -501,7 +502,7 @@ void FModelInterface::UpdateCLFromP4Async() const
 
 bool FModelInterface::GetInputEnabled()
 {
-	return SubmitToolState == ESubmitToolAppState::WaitingUserInput || SubmitToolState == ESubmitToolAppState::SubmitLocked;
+	return SubmitToolState == ESubmitToolAppState::WaitingUserInput;
 }
 
 void FModelInterface::RequestPreflight(bool bForceStart)
@@ -722,25 +723,25 @@ void FModelInterface::StartSubmitProcess(bool bSkipShelfDialog)
 						LocalList.Add(LocalFiles[i]);
 					}
 				}
-				FString LocalListString = FString::Join(LocalList, TEXT("\n - "));
+				FString LocalListString = FSubmitToolCoreUtils::StringBuilderJoin<1024>(LocalList, TEXT("\n - "));
 				if (LocalFiles.Num() > MaxFilesToList)
 				{
 					LocalListString = FString::Printf(TEXT("%s\n - And %d other files"), *LocalListString, LocalFiles.Num() - MaxFilesToList);
 				}
 
-				FString ShelveListString = FString::Join(ShelvedList, TEXT("\n - "));
+				FString ShelveListString = FSubmitToolCoreUtils::StringBuilderJoin<1024>(ShelvedList, TEXT("\n - "));
 				if (ShelvedFiles.Num() > MaxFilesToList)
 				{
 					ShelveListString = FString::Printf(TEXT("%s\n - And %d other files"), *ShelveListString, ShelvedFiles.Num() - MaxFilesToList);
 				}
 			}
 
-			ChangeState(ESubmitToolAppState::Submitting, SubmitToolState == ESubmitToolAppState::SubmitLocked && bIsUserInAllowlist);
+			ChangeState(ESubmitToolAppState::Submitting, bIsStreamLocked && bIsUserInAllowlist);
 			DeleteShelvedFiles();
 		}
 		else
 		{
-			ChangeState(ESubmitToolAppState::Submitting, SubmitToolState == ESubmitToolAppState::SubmitLocked && bIsUserInAllowlist);
+			ChangeState(ESubmitToolAppState::Submitting, bIsStreamLocked && bIsUserInAllowlist);
 			RevertUnchangedAndSubmit();
 		}
 	}
@@ -783,23 +784,24 @@ void FModelInterface::RefreshStateBasedOnFiles()
 			{
 				// Copy the shelved file depot paths on the main thread since they are not safe to reference from other threads
 				UE::Tasks::Launch(UE_SOURCE_LOCATION, [this, ShelvedFilesDepotPaths = ChangelistService->GetShelvedFilesDepotPaths()] {
-					FPathLockdownResult LockdownStatus = P4LockdownService->ArePathsInLockdown(ShelvedFilesDepotPaths);
-					if (LockdownStatus.bIsLocked)
+					FPathLockdownResult LockdownStatus = P4LockdownService->ArePathsInLockdown(ShelvedFilesDepotPaths);					
+					UE::Tasks::Launch(UE_SOURCE_LOCATION, [this, LockdownStatus = MoveTemp(LockdownStatus)]
 					{
-						UE_LOG(LogSubmitTool, Log, TEXT("There are no local files in CL %s, Submit is disabled but you can still request an Integration with your shelved files"), *ChangelistService->GetCLID());
-						UE::Tasks::Launch(UE_SOURCE_LOCATION, [this, LockdownStatus = MoveTemp(LockdownStatus)]
+						bIsStreamLocked = LockdownStatus.bIsLocked;
+						bIsUserInAllowlist = LockdownStatus.bAllowlisted;
+						TitleMessage = FSubmitToolCoreUtils::StringBuilderJoin(LockdownStatus.Messages, TEXT("\n"));
+						if (LockdownStatus.bIsLocked)
 						{
-							bIsUserInAllowlist = LockdownStatus.bAllowlisted;
-							TitleMessage = LockdownStatus.Message;
-							ChangeState(ESubmitToolAppState::SubmitLocked);
-						}, LowLevelTasks::ETaskPriority::Normal, UE::Tasks::EExtendedTaskPriority::GameThreadNormalPri);
-					}
-					else
-					{
-						UE_LOG(LogSubmitTool, Error, TEXT("There are no files in CL %s, SUBMIT IS DISABLED"), *ChangelistService->GetCLID());
-						AsyncTask(ENamedThreads::GameThread, [] { ChangeState(ESubmitToolAppState::Errored); });
-					}
-					});
+							UE_LOG(LogSubmitTool, Log, TEXT("There are no local files in CL %s, Submit is disabled but you can still request an Integration with your shelved files"), *ChangelistService->GetCLID());
+							ChangeState(ESubmitToolAppState::WaitingUserInput);
+						}
+						else
+						{
+							UE_LOG(LogSubmitTool, Error, TEXT("There are no files in CL %s, SUBMIT IS DISABLED"), *ChangelistService->GetCLID());
+							ChangeState(ESubmitToolAppState::Errored);
+						}
+					}, LowLevelTasks::ETaskPriority::Normal, UE::Tasks::EExtendedTaskPriority::GameThreadNormalPri);
+				});
 			}
 		}
 	}
@@ -841,20 +843,14 @@ void FModelInterface::RefreshStateBasedOnFiles()
 		// Copy the file depot paths on the main thread since they are not safe to reference from other threads
 		UE::Tasks::Launch(UE_SOURCE_LOCATION, [this, FilesDepotPaths = ChangelistService->GetOpenedFilesDepotPaths()] {
 			FPathLockdownResult LockdownStatus = P4LockdownService->ArePathsInLockdown(FilesDepotPaths);
-			if (LockdownStatus.bIsLocked)
+			UE::Tasks::Launch(UE_SOURCE_LOCATION, [this, LockdownStatus = MoveTemp(LockdownStatus)]
 			{
-				UE::Tasks::Launch(UE_SOURCE_LOCATION, [this, LockdownStatus = MoveTemp(LockdownStatus)]
-				{
-					bIsUserInAllowlist = LockdownStatus.bAllowlisted;
-					TitleMessage = LockdownStatus.Message;
-					ChangeState(ESubmitToolAppState::SubmitLocked);
-				}, LowLevelTasks::ETaskPriority::Normal, UE::Tasks::EExtendedTaskPriority::GameThreadNormalPri);
-			}
-			else
-			{
-				AsyncTask(ENamedThreads::GameThread, [] { ChangeState(ESubmitToolAppState::WaitingUserInput); });
-			}
-			});
+				bIsStreamLocked = LockdownStatus.bIsLocked;
+				bIsUserInAllowlist = LockdownStatus.bAllowlisted;
+				TitleMessage = FSubmitToolCoreUtils::StringBuilderJoin(LockdownStatus.Messages, TEXT("\n"));
+				ChangeState(ESubmitToolAppState::WaitingUserInput);
+			}, LowLevelTasks::ETaskPriority::Normal, UE::Tasks::EExtendedTaskPriority::GameThreadNormalPri);
+		});
 	}
 }
 
@@ -1219,7 +1215,7 @@ void FModelInterface::OnGetUsersFromSwarmCompleted(const TUniquePtr<FSwarmReview
 
 			if (bApplied)
 			{
-				UE_LOG(LogSubmitTool, Log, TEXT("RB tag set to users that upvoted review '%d' Users: %s"), InReview->Id, *FString::Join(SwarmUserValues, TEXT(", ")));
+				UE_LOG(LogSubmitTool, Log, TEXT("RB tag set to users that upvoted review '%d' Users: %s"), InReview->Id, *FSubmitToolCoreUtils::StringBuilderJoin<128>(SwarmUserValues, TEXT(", ")));
 				UE_LOG(LogSubmitToolDebug, Log, TEXT("Re-running Tag validator after applying the #rb from swarm"));
 				ValidateCLDescription();
 			}
diff --git a/Engine/Source/Programs/SubmitTool/Private/Models/ModelInterface.h b/Engine/Source/Programs/SubmitTool/Private/Models/ModelInterface.h
--- a/Engine/Source/Programs/SubmitTool/Private/Models/ModelInterface.h
+++ b/Engine/Source/Programs/SubmitTool/Private/Models/ModelInterface.h
@@ -31,7 +31,6 @@ enum class ESubmitToolAppState : uint8 {
 	Errored = 3,
 	P4BlockingOperation = 4,
 	Submitting = 5,
-	SubmitLocked = 6,
 	Finished = 7,
 };
 
@@ -45,35 +44,24 @@ namespace SubmitToolAppState {
 		{ ESubmitToolAppState::Initializing, StateList {
 				ESubmitToolAppState::WaitingUserInput,
 				ESubmitToolAppState::P4BlockingOperation,
-				ESubmitToolAppState::Errored,
-				ESubmitToolAppState::SubmitLocked
+				ESubmitToolAppState::Errored
 		}},
 		{ ESubmitToolAppState::WaitingUserInput, StateList {
 				ESubmitToolAppState::Submitting,
 				ESubmitToolAppState::Errored,
-				ESubmitToolAppState::P4BlockingOperation,
-				ESubmitToolAppState::SubmitLocked
+				ESubmitToolAppState::P4BlockingOperation
 		}},
 		{ ESubmitToolAppState::Errored, StateList {
-			ESubmitToolAppState::WaitingUserInput,
-			ESubmitToolAppState::SubmitLocked
+			ESubmitToolAppState::WaitingUserInput
 		}},
 		{ ESubmitToolAppState::P4BlockingOperation, StateList {
 				ESubmitToolAppState::WaitingUserInput,
 				ESubmitToolAppState::Errored,
-				ESubmitToolAppState::SubmitLocked,
 				ESubmitToolAppState::Finished
 		}},
 		{ ESubmitToolAppState::Submitting, StateList {
 				ESubmitToolAppState::WaitingUserInput,
 				ESubmitToolAppState::Errored,
-				ESubmitToolAppState::SubmitLocked,
-				ESubmitToolAppState::Finished
-		}},
-		{ ESubmitToolAppState::SubmitLocked, StateList {
-				ESubmitToolAppState::WaitingUserInput,
-				ESubmitToolAppState::SubmitLocked,
-				ESubmitToolAppState::Errored,
 				ESubmitToolAppState::Finished
 		}},
 		{ ESubmitToolAppState::Finished, StateList {
@@ -174,7 +162,6 @@ public:
 	static bool GetInputEnabled();
 	static void SetErrorState()														{ ChangeState(ESubmitToolAppState::Errored);}
 	static const ESubmitToolAppState GetState()										{ return SubmitToolState; }
-	bool IsSubmitBlocked() const													{ return SubmitToolState == ESubmitToolAppState::SubmitLocked; }
 	void DeleteShelvedFiles() const													{ ChangelistService->DeleteShelvedFiles(DeleteShelveCallback); }
 
 	void RequestPreflight(bool bForceStart = false);
@@ -201,11 +188,12 @@ public:
 	TWeakPtr<FPreflightService> GetPreflightService() { return this->PreflightService; }
 
 	bool bIsUserInAllowlist = false;
+	bool bIsStreamLocked = false;
 
 	const FSubmitToolParameters& GetParameters() const { return Parameters; }
 
 	const FString& GetTitleMessage() const;
-	bool IsIntegrationRequired() const { return GetState() == ESubmitToolAppState::SubmitLocked; };
+	bool IsIntegrationRequired() const { return bIsStreamLocked; };
 	void RequestIntegration() const;
 	const TMap<FString, TSharedPtr<FIntegrationOptionBase>>& GetIntegrationOptions() { return IntegrationService->GetIntegrationOptions(); }
 	bool ValidateIntegrationOptions(bool bSilent) const	{ return IntegrationService->ValidateIntegrationOptions(bSilent); }
