[SubmitTool] Initializing FCommandLine in Linux and Mac

#Synced-CL 48728512

diff --git a/Engine/Source/Programs/SubmitTool/Private/Platforms/Linux/SubmitToolMainLinux.cpp b/Engine/Source/Programs/SubmitTool/Private/Platforms/Linux/SubmitToolMainLinux.cpp
--- a/Engine/Source/Programs/SubmitTool/Private/Platforms/Linux/SubmitToolMainLinux.cpp
+++ b/Engine/Source/Programs/SubmitTool/Private/Platforms/Linux/SubmitToolMainLinux.cpp
@@ -1,10 +1,12 @@
 // Copyright Epic Games, Inc. All Rights Reserved.
 
 #include "SubmitToolApp.h"
+#include "Misc/CommandLine.h"
 #include "UnixCommonStartup.h"
 
 int RunSubmitToolWrapper(const TCHAR* Commandline)
 {
+	FCommandLine::Set(Commandline);
     return RunSubmitTool(Commandline);
 }
 
diff --git a/Engine/Source/Programs/SubmitTool/Private/Platforms/Mac/SubmitToolMainMac.cpp b/Engine/Source/Programs/SubmitTool/Private/Platforms/Mac/SubmitToolMainMac.cpp
--- a/Engine/Source/Programs/SubmitTool/Private/Platforms/Mac/SubmitToolMainMac.cpp
+++ b/Engine/Source/Programs/SubmitTool/Private/Platforms/Mac/SubmitToolMainMac.cpp
@@ -6,6 +6,7 @@
 
 int RunSubmitToolWrapper(const TCHAR* Commandline)
 {
+	FCommandLine::Set(Commandline);
 	return RunSubmitTool(Commandline);
 }
 
