[Submit Tool] Fixing menu handler losing losing reference to ModelInterface causing a crash when clicking 'Help'

#Synced-CL 48728512

diff --git a/Engine/Source/Programs/SubmitTool/Private/Platforms/Mac/SubmitToolMainMac.cpp b/Engine/Source/Programs/SubmitTool/Private/Platforms/Mac/SubmitToolMainMac.cpp
--- a/Engine/Source/Programs/SubmitTool/Private/Platforms/Mac/SubmitToolMainMac.cpp
+++ b/Engine/Source/Programs/SubmitTool/Private/Platforms/Mac/SubmitToolMainMac.cpp
@@ -2,6 +2,7 @@
 
 #include "SubmitToolApp.h"
 #include "Mac/MacProgramDelegate.h"
+#include "Misc/CommandLine.h"
 #include "LaunchEngineLoop.h"
 
 int RunSubmitToolWrapper(const TCHAR* Commandline)
