[SubmitTool] Syncing up SubmitTool.ini config files from the depot before starting up

#Synced-CL 48728512

diff --git a/Engine/Source/Programs/SubmitTool/Private/Logic/Services/Interfaces/ISTSourceControlService.h b/Engine/Source/Programs/SubmitTool/Private/Logic/Services/Interfaces/ISTSourceControlService.h
--- a/Engine/Source/Programs/SubmitTool/Private/Logic/Services/Interfaces/ISTSourceControlService.h
+++ b/Engine/Source/Programs/SubmitTool/Private/Logic/Services/Interfaces/ISTSourceControlService.h
@@ -5,6 +5,7 @@
 #include "CoreMinimal.h"
 #include "ISubmitToolService.h"
 #include "ISourceControlProvider.h"
+#include "Parameters/SubmitToolParameters.h"
 #include "Memory/SharedBuffer.h"
 #include "Tasks/Task.h"
 #include "Misc/Base64.h"
@@ -83,6 +84,7 @@ struct FSCCommand
 	FString Command;
 	TArray<FString> Args;
 	bool bCollectData = false;
+	bool bSilentErrors = false;
 	// out
 	FSCCRecordSet Values;
 	FSourceControlResultInfo Results;
@@ -99,6 +101,7 @@ using FOnSCCCommandComplete = TDelegate<void(const FSCCommand&)>;
 class ISTSourceControlService : public ISubmitToolService
 {
 public:
+	virtual void InitializeParameters(const FSubmitToolParameters& InParameters) = 0;
 	virtual const TUniquePtr<ISourceControlProvider>& GetProvider() const = 0;
 	virtual void GetUsers(const FOnUsersGet::FDelegate& Callback) = 0;
 	virtual void GetGroups(const FOnGroupsGet::FDelegate& Callback) = 0;
@@ -126,11 +129,11 @@ public:
 	virtual bool IsBusy() const = 0;
 
 	// Runs the command in a task thread and the provided callback in a dependent GT task.
-	void RunCommand(FString&& InCommand, TArray<FString>&& InAdditionalArgs, FOnSCCCommandComplete InCompleteCallback)
+	void RunCommand(FString&& InCommand, TArray<FString>&& InAdditionalArgs, FOnSCCCommandComplete InCompleteCallback, bool bSilentErrors = false)
 	{
 		check(InCompleteCallback.IsBound());
-		TSharedRef<FSCCommand> Command = MakeShared<FSCCommand>(MoveTemp(InCommand), MoveTemp(InAdditionalArgs));
-		UE::Tasks::FTask CommandTask = UE::Tasks::Launch(UE_SOURCE_LOCATION, [this, Command, InCompleteCallback]
+		TSharedRef<FSCCommand> Command = MakeShared<FSCCommand>(MoveTemp(InCommand), MoveTemp(InAdditionalArgs), false, bSilentErrors);
+		UE::Tasks::FTask CommandTask = UE::Tasks::Launch(UE_SOURCE_LOCATION, [this, Command, InCompleteCallback, bSilentErrors]
 		{
 			Command->bSuccess = RunCommandInternal(*Command);
 			UE::Tasks::Launch(UE_SOURCE_LOCATION, [InCompleteCallback, Command]
@@ -142,10 +145,10 @@ public:
 	}
 
 	// Runs the command in a task thread, calling GetResult() will wait until the task is completed and return the command data.
-	UE::Tasks::TTask<TSharedRef<FSCCommand>> RunCommand(FString&& InCommand, TArray<FString>&& InAdditionalArgs)
+	UE::Tasks::TTask<TSharedRef<FSCCommand>> RunCommand(FString&& InCommand, TArray<FString>&& InAdditionalArgs, bool bSilentErrors = false)
 	{
-		TSharedRef<FSCCommand> Command = MakeShared<FSCCommand>(MoveTemp(InCommand), MoveTemp(InAdditionalArgs));
-		return UE::Tasks::Launch(UE_SOURCE_LOCATION, [this, Command]
+		TSharedRef<FSCCommand> Command = MakeShared<FSCCommand>(MoveTemp(InCommand), MoveTemp(InAdditionalArgs), false, bSilentErrors);
+		return UE::Tasks::Launch(UE_SOURCE_LOCATION, [this, Command, bSilentErrors]
 		{
 			Command->bSuccess = RunCommandInternal(*Command);
 			return Command;
@@ -153,9 +156,9 @@ public:
 	}
 
 	// Runs the command synchronously and returns the command data
-	FSCCommand RunSyncCommand(FString&& InCommand, TArray<FString>&& InAdditionalArgs)
+	FSCCommand RunSyncCommand(FString&& InCommand, TArray<FString>&& InAdditionalArgs, bool bSilentErrors = false)
 	{
-		FSCCommand Command(MoveTemp(InCommand), MoveTemp(InAdditionalArgs));
+		FSCCommand Command(MoveTemp(InCommand), MoveTemp(InAdditionalArgs), false, bSilentErrors);
 		Command.bSuccess = RunCommandInternal(Command);
 		return MoveTemp(Command);
 	}
diff --git a/Engine/Source/Programs/SubmitTool/Private/Logic/Services/SourceControl/PerforceData.h b/Engine/Source/Programs/SubmitTool/Private/Logic/Services/SourceControl/PerforceData.h
--- a/Engine/Source/Programs/SubmitTool/Private/Logic/Services/SourceControl/PerforceData.h
+++ b/Engine/Source/Programs/SubmitTool/Private/Logic/Services/SourceControl/PerforceData.h
@@ -18,7 +18,7 @@ THIRD_PARTY_INCLUDES_END
 #define LOCTEXT_NAMESPACE "SubmitToolPerforce"
 
 UENUM()
-enum class EP4ClientUserFlags
+enum class EP4ClientUserFlags : uint64
 {
 	None = 0,
 	/** The server uses unicode */
@@ -27,7 +27,8 @@ enum class EP4ClientUserFlags
 	CollectData = 1 << 1,
 	UseZTag = 1 << 2,
 	UseClient = 1 << 3,
-	UseUser = 1 << 4
+	UseUser = 1 << 4,
+	SilentErrors = 1 << 5,
 };
 
 ENUM_CLASS_FLAGS(EP4ClientUserFlags)
diff --git a/Engine/Source/Programs/SubmitTool/Private/Logic/Services/SourceControl/SubmitToolPerforce.cpp b/Engine/Source/Programs/SubmitTool/Private/Logic/Services/SourceControl/SubmitToolPerforce.cpp
--- a/Engine/Source/Programs/SubmitTool/Private/Logic/Services/SourceControl/SubmitToolPerforce.cpp
+++ b/Engine/Source/Programs/SubmitTool/Private/Logic/Services/SourceControl/SubmitToolPerforce.cpp
@@ -105,8 +105,7 @@ size_t FConnectionPool::CreateConnection()
 	return P4Connections.Add(MakeUnique<FClientApiWrapper>(MoveTemp(P4Client), FOnIsCancelled::CreateLambda([this] { return bWantsCancel.load(); })));
 }
 
-FSubmitToolPerforce::FSubmitToolPerforce(const FSubmitToolParameters& InParameters)
-: Parameters(InParameters)
+FSubmitToolPerforce::FSubmitToolPerforce()
 {
 	TRACE_CPUPROFILER_EVENT_SCOPE(FSubmitToolPerforce::FSubmitToolPerforce);
 
@@ -168,6 +167,11 @@ FSubmitToolPerforce::~FSubmitToolPerforce()
 	FTSTicker::RemoveTicker(TickHandle);
 }
 
+void FSubmitToolPerforce::InitializeParameters(const FSubmitToolParameters& InParameters)
+{
+	GroupsToExclude = InParameters.GeneralParameters.GroupsToExclude;
+}
+
 const TUniquePtr<ISourceControlProvider>& FSubmitToolPerforce::GetProvider() const
 {
 	return SCCProvider;
@@ -187,6 +191,7 @@ UE::Tasks::TTask<TSharedRef<FSCCommand>> FSubmitToolPerforce::GetStream(const FS
 
 void FSubmitToolPerforce::ParseStreamCommand(const FSCCommand& InCompletedCommand, const FString& InRequestedStream, bool bParseHierarchy)
 {
+	TRACE_CPUPROFILER_EVENT_SCOPE(FSubmitToolPerforce::ParseStreamCommand);
 	if (InCompletedCommand.bSuccess && InCompletedCommand.Values.Num() > 0 && InCompletedCommand.Values[0].Contains(TEXT("Stream")))
 	{
 		TUniquePtr<FSCCStream> Stream = MakeUnique<FSCCStream>(InCompletedCommand.Values[0][TEXT("Stream")], InCompletedCommand.Values[0][TEXT("Parent")], InCompletedCommand.Values[0][TEXT("Type")]);
@@ -562,7 +567,7 @@ void FSubmitToolPerforce::GetGroups(const FOnGroupsGet::FDelegate& Callback)
 
 				TArray<TSharedPtr<FString>> GroupsResults;
 
-				FGroupsClientUser P4User(ResultValues, bIsUnicodeServer ? EP4ClientUserFlags::UnicodeServer : EP4ClientUserFlags::None, ResultInfo, GroupsResults, Parameters.GeneralParameters.GroupsToExclude);
+				FGroupsClientUser P4User(ResultValues, bIsUnicodeServer ? EP4ClientUserFlags::UnicodeServer : EP4ClientUserFlags::None, ResultInfo, GroupsResults, GroupsToExclude);
 
 				{
 					TRACE_CPUPROFILER_EVENT_SCOPE(FSubmitToolPerforce::RunP4Command_groups);
@@ -770,7 +775,13 @@ void FSubmitToolPerforce::GetUsersAndGroups(const FOnUsersAndGroupsGet::FDelegat
 
 bool FSubmitToolPerforce::RunCommandInternal(FSCCommand& Command)
 {
-	return RunP4Command(Command);
+	EP4ClientUserFlags Flags = DefaultFlags;
+	if (Command.bSilentErrors)
+	{
+		Flags |= EP4ClientUserFlags::SilentErrors;
+	}
+
+	return RunP4Command(Command, Flags);
 }
 
 const TArray<FSCCStream*>& FSubmitToolPerforce::GetClientStreams() const
@@ -985,7 +996,14 @@ bool FSubmitToolPerforce::RunP4Command(FSCCommand& Command, EP4ClientUserFlags C
 	{
 		for(const FText& Error : Command.Results.ErrorMessages)
 		{
-			UE_LOG(LogSubmitToolP4, Error, TEXT("%s: %s"), *FullCommand, *Error.ToString());
+			if (EnumHasAnyFlags(CmdFlags, EP4ClientUserFlags::SilentErrors))
+			{
+				UE_LOG(LogSubmitToolP4Debug, Error, TEXT("%s: %s"), *FullCommand, *Error.ToString());
+			}
+			else
+			{
+				UE_LOG(LogSubmitToolP4, Error, TEXT("%s: %s"), *FullCommand, *Error.ToString());
+			}
 		}
 	}
 
diff --git a/Engine/Source/Programs/SubmitTool/Private/Logic/Services/SourceControl/SubmitToolPerforce.h b/Engine/Source/Programs/SubmitTool/Private/Logic/Services/SourceControl/SubmitToolPerforce.h
--- a/Engine/Source/Programs/SubmitTool/Private/Logic/Services/SourceControl/SubmitToolPerforce.h
+++ b/Engine/Source/Programs/SubmitTool/Private/Logic/Services/SourceControl/SubmitToolPerforce.h
@@ -3,7 +3,6 @@
 #pragma once
 
 #include "Logic/Services/Interfaces/ISTSourceControlService.h"
-#include "Parameters/SubmitToolParameters.h"
 #include "PerforceData.h"
 #include "Containers/Ticker.h"
 #include "Containers/LockFreeFixedSizeAllocator.h"
@@ -73,9 +72,11 @@ private:
 class FSubmitToolPerforce final : public ISTSourceControlService
 {
 public:
-	FSubmitToolPerforce(const FSubmitToolParameters& InParameters);
+	FSubmitToolPerforce();
 	~FSubmitToolPerforce();
 
+
+	virtual void InitializeParameters(const FSubmitToolParameters& InParameters);
 	virtual const TUniquePtr<ISourceControlProvider>& GetProvider() const override;
 	virtual void GetUsers(const FOnUsersGet::FDelegate& Callback) override;
 	virtual void GetGroups(const FOnGroupsGet::FDelegate& Callback) override;
@@ -115,7 +116,7 @@ private:
 	void ParseStreamCommand(const FSCCommand& InCompletedCommand, const FString& InRequestedStream = FString(), bool bParseHierarchy = true);
 	void WaitForInitialStreamTask();
 
-	const FSubmitToolParameters& Parameters;
+	TArray<FString> GroupsToExclude;
 	TMap<FString, TSharedPtr<FUserData>> CachedUsers;
 	TArray<TSharedPtr<FUserData>> CachedUsersArray;
 	TArray<TSharedPtr<FString>> CachedGroupsArray;
diff --git a/Engine/Source/Programs/SubmitTool/Private/Logic/TasksService.cpp b/Engine/Source/Programs/SubmitTool/Private/Logic/TasksService.cpp
--- a/Engine/Source/Programs/SubmitTool/Private/Logic/TasksService.cpp
+++ b/Engine/Source/Programs/SubmitTool/Private/Logic/TasksService.cpp
@@ -33,6 +33,7 @@ FTasksService::~FTasksService()
 
 void FTasksService::InitializeTasks(const TArray<TSharedRef<FValidatorBase>>& InTasks)
 {
+	TRACE_CPUPROFILER_EVENT_SCOPE(FTasksService::InitializeTasks);
 	for(const TSharedRef<FValidatorBase>& Task : InTasks)
 	{
 		if(!Task->Activate())
diff --git a/Engine/Source/Programs/SubmitTool/Private/Models/ModelInterface.cpp b/Engine/Source/Programs/SubmitTool/Private/Models/ModelInterface.cpp
--- a/Engine/Source/Programs/SubmitTool/Private/Models/ModelInterface.cpp
+++ b/Engine/Source/Programs/SubmitTool/Private/Models/ModelInterface.cpp
@@ -38,8 +38,8 @@
 ESubmitToolAppState FModelInterface::SubmitToolState = ESubmitToolAppState::Initializing;
 FOnStateChanged FModelInterface::OnStateChanged;
 
-FModelInterface::FModelInterface(const FSubmitToolParameters& InParameters) :
-	Parameters(InParameters)
+FModelInterface::FModelInterface(const FSubmitToolParameters& InParameters, TSharedPtr<ISTSourceControlService>& InSourceControlService) :
+	Parameters(InParameters), SourceControlService(InSourceControlService)
 {
 	TRACE_CPUPROFILER_EVENT_SCOPE(FModelInterface::FModelInterface);
 
@@ -61,7 +61,7 @@ FModelInterface::FModelInterface(const FSubmitToolParameters& InParameters) :
 		ServiceProvider->RegisterService<ICacheDataService>(MakeShared<FCacheDataService>(Parameters.GeneralParameters));
 	}
 
-	SourceControlService = MakeShared<FSubmitToolPerforce>(InParameters);
+	SourceControlService->InitializeParameters(InParameters);
 	ServiceProvider->RegisterService<ISTSourceControlService>(SourceControlService.ToSharedRef());
 	ValidationService = MakeShared<FTasksService>(InParameters.Validators, TEXT("SubmitTool.StandAlone.Validator"));
 	ServiceProvider->RegisterService<FTasksService>(ValidationService.ToSharedRef(), TEXT("ValidationService"));
@@ -198,6 +198,7 @@ void FModelInterface::Dispose() const
 
 void FModelInterface::ParseValidators() const
 {
+	TRACE_CPUPROFILER_EVENT_SCOPE(FModelInterface::ParseValidators);
 	TArray<TSharedRef<FValidatorBase>> Tasks;
 
 	for (const TPair<FName, FString> DefinitionPair : Parameters.Validators)
@@ -277,6 +278,7 @@ void FModelInterface::ParseValidators() const
 
 void FModelInterface::ParsePreSubmitOperations() const
 {
+	TRACE_CPUPROFILER_EVENT_SCOPE(FModelInterface::ParsePreSubmitOperations);
 	TArray<TSharedRef<FValidatorBase>> Tasks;
 
 	for (const TPair<FName, FString> DefinitionPair : Parameters.PresubmitOperations)
diff --git a/Engine/Source/Programs/SubmitTool/Private/Models/ModelInterface.h b/Engine/Source/Programs/SubmitTool/Private/Models/ModelInterface.h
--- a/Engine/Source/Programs/SubmitTool/Private/Models/ModelInterface.h
+++ b/Engine/Source/Programs/SubmitTool/Private/Models/ModelInterface.h
@@ -77,7 +77,7 @@ DECLARE_MULTICAST_DELEGATE_TwoParams(FOnStateChanged, const ESubmitToolAppState
 class FModelInterface
 {
 public:
-	FModelInterface(const FSubmitToolParameters& InParameters);
+	FModelInterface(const FSubmitToolParameters& InParameters, TSharedPtr<ISTSourceControlService>& InSourceControlService);
 	~FModelInterface();
 	void Dispose() const;
 	void ParseValidators() const;
diff --git a/Engine/Source/Programs/SubmitTool/Private/SubmitToolApp.cpp b/Engine/Source/Programs/SubmitTool/Private/SubmitToolApp.cpp
--- a/Engine/Source/Programs/SubmitTool/Private/SubmitToolApp.cpp
+++ b/Engine/Source/Programs/SubmitTool/Private/SubmitToolApp.cpp
@@ -4,6 +4,7 @@
 #include "CoreMinimal.h"
 #include "SubmitToolCoreUtils.h"
 #include "Models/SubmitToolUserPrefs.h"
+#include "Logic/Services/SourceControl/SubmitToolPerforce.h"
 
 #include "GenericPlatform/GenericPlatformMisc.h"
 #include "RequiredProgramMainCPPInclude.h"
@@ -206,28 +207,38 @@ int RunSubmitTool(const TCHAR* CommandLine)
 	// initialize the configuration system
 	FConfiguration::Init();
 
-	// allow running multiple parallel instances of UAT
-	FPlatformMisc::SetEnvironmentVar(TEXT("uebp_UATMutexNoWait"), TEXT("1"));
-	UE_LOG(LogSubmitToolDebug, Log, TEXT("Set environment variable uebp_UATMutexNoWait=1 to allow running multiple parallel instances of UAT."));
-
 	// App Scope
 	{
 		UE_LOG(LogSubmitTool, Log, TEXT("Version %s"), *FAppVersion::GetVersion());
+		TSharedPtr<ISTSourceControlService> SourceControlService = MakeShared<FSubmitToolPerforce>();
+		
+		// Sync up the config as early as possible
+		UE::Tasks::TTask<TSharedRef<FSCCommand>> SyncCommand = SyncUpConfigFiles(SourceControlService);
 		TUniquePtr<FSubmitToolUserPrefs> UserPrefs = FSubmitToolUserPrefs::Initialize(GetUserPrefsPath());
 
+		// allow running multiple parallel instances of UAT
+		FPlatformMisc::SetEnvironmentVar(TEXT("uebp_UATMutexNoWait"), TEXT("1"));
+		UE_LOG(LogSubmitToolDebug, Log, TEXT("Set environment variable uebp_UATMutexNoWait=1 to allow running multiple parallel instances of UAT."));
+
+		{
+			TRACE_CPUPROFILER_EVENT_SCOPE(FSubmitToolApp::InitializeUI);
+			// crank up a normal Slate application using the platform's standalone renderer while the sync is in progress
+			FSlateApplication::InitializeAsStandaloneApplication(GetStandardStandaloneRenderer());
+			FAppStyle::SetAppStyleSet(FSubmitToolStyle::Get());
+			FSlateApplication::InitHighDPI(true);
+		}
+
+		// Wait for the sync since everything else needs the config.
+		SyncCommand.Wait();
+		SyncCommand = { };
+
 		FString ParameterFile;
 		CmdLineParameters.GetValue(FSubmitToolCmdLine::ParameterFile, ParameterFile);
 		FSubmitToolParametersBuilder ParametersBuilder(ParameterFile);
 		FSubmitToolParameters Parameters = ParametersBuilder.Build();
 
 		// Create a new instance of model Interface so that UI can communicate
-		// Do this early to kick off any async tasks while we are taking time to build and initialize the actual UI app
-		TUniquePtr<FModelInterface> ModelInterface = MakeUnique<FModelInterface>(Parameters);
-
-		// crank up a normal Slate application using the platform's standalone renderer
-		FSlateApplication::InitializeAsStandaloneApplication(GetStandardStandaloneRenderer());
-		FAppStyle::SetAppStyleSet(FSubmitToolStyle::Get());
-		FSlateApplication::InitHighDPI(true);
+		TUniquePtr<FModelInterface> ModelInterface = MakeUnique<FModelInterface>(Parameters, SourceControlService);	
 
 		// Build the slate ui
 		TUniquePtr<FSlateUI> SlateUI = MakeSlateUI(ModelInterface.Get());
@@ -257,6 +268,50 @@ int RunSubmitTool(const TCHAR* CommandLine)
 	return 0;
 }
 
+UE::Tasks::TTask<TSharedRef<FSCCommand>> SyncUpConfigFiles(const TSharedPtr<ISTSourceControlService>& InSourceControlService)
+{
+	TRACE_CPUPROFILER_EVENT_SCOPE(FSubmitToolApp::SyncUpConfigFiles);
+	const FString Root = FConfiguration::Substitute(TEXT("$(root)"));
+	TArray<FString> Files{
+		Root / TEXT("Engine/Restricted/NotForLicensees/Config/Mac/MacSubmitTool.ini"),
+		Root / TEXT("Engine/Restricted/NotForLicensees/Config/Linux/LinuxSubmitTool.ini"),
+		Root / TEXT("Engine/Restricted/NotForLicensees/Config/Windows/WindowsSubmitTool.ini"),
+		Root / TEXT("Engine/Restricted/NotForLicensees/Config/SubmitTool.ini")
+	};
+
+	//Files.Add(Root / TEXT("*/Config/Mac/MacSubmitTool.ini"));
+	//Files.Add(Root / TEXT("*/Config/Linux/LinuxSubmitTool.ini"));
+	//Files.Add(Root / TEXT("*/Config/Windows/WindowsSubmitTool.ini"));
+	//Files.Add(Root / TEXT("*/Config/SubmitTool.ini"));
+
+	IFileManager::Get().IterateDirectory(*Root,
+		[&Files](const TCHAR* FileOrDir, bool bIsDir)
+		{
+			if (bIsDir)
+			{
+				FString Dir = { FileOrDir };
+				if (!Dir.Contains(TEXT("Engine"), ESearchCase::IgnoreCase))
+				{
+					const FString Extension = TEXT(".uproject");
+					TArray<FString> UProjects;
+					IFileManager::Get().FindFiles(UProjects, FileOrDir, *Extension);
+
+					if (UProjects.Num() != 0)
+					{
+						Files.Add(Dir / TEXT("Config/Mac/MacSubmitTool.ini"));
+						Files.Add(Dir / TEXT("Config/Linux/LinuxSubmitTool.ini"));
+						Files.Add(Dir / TEXT("Config/Windows/WindowsSubmitTool.ini"));
+						Files.Add(Dir / TEXT("Config/SubmitTool.ini"));
+					}
+				}
+			}
+			return true;
+		}
+	);
+
+	return InSourceControlService->RunCommand(TEXT("sync"), MoveTemp(Files), true);
+}
+
 FString GetUserPrefsPath()
 {
 	return FPaths::Combine(FSubmitToolCoreUtils::GetLocalAppDataPath(), "SubmitTool", "SubmitToolPrefs.ini");
diff --git a/Engine/Source/Programs/SubmitTool/Public/SubmitToolApp.h b/Engine/Source/Programs/SubmitTool/Public/SubmitToolApp.h
--- a/Engine/Source/Programs/SubmitTool/Public/SubmitToolApp.h
+++ b/Engine/Source/Programs/SubmitTool/Public/SubmitToolApp.h
@@ -3,9 +3,14 @@
 #pragma once
 
 #include "CoreMinimal.h"
+#include "Tasks/Task.h"
+
+class ISTSourceControlService;
+struct FSCCommand;
 
 /**
  * Run the SubmitTool .
  */
 int RunSubmitTool(const TCHAR* Commandline);
+UE::Tasks::TTask<TSharedRef<FSCCommand>> SyncUpConfigFiles(const TSharedPtr<ISTSourceControlService>& InSourceControlService);
 FString GetUserPrefsPath();
