[Submit Tool] 'Show Review' and 'Open Preflight URL' will be enabled even when input is being blocked

#Synced-CL 48728512

diff --git a/Engine/Source/Programs/SubmitTool/Private/View/Widgets/TagSectionWidget.cpp b/Engine/Source/Programs/SubmitTool/Private/View/Widgets/TagSectionWidget.cpp
--- a/Engine/Source/Programs/SubmitTool/Private/View/Widgets/TagSectionWidget.cpp
+++ b/Engine/Source/Programs/SubmitTool/Private/View/Widgets/TagSectionWidget.cpp
@@ -55,7 +55,6 @@ void STagSectionWidget::Construct(const FArguments& InArgs)
 			.ModelInterface(InArgs._ModelInterface)
 			.Tag(TagInstance)
 			.JiraWidget(Issues)
-			.IsEnabled_Static(&FModelInterface::GetInputEnabled)
 		];
 
 		TagCount++;
diff --git a/Engine/Source/Programs/SubmitTool/Private/View/Widgets/TagWidget.cpp b/Engine/Source/Programs/SubmitTool/Private/View/Widgets/TagWidget.cpp
--- a/Engine/Source/Programs/SubmitTool/Private/View/Widgets/TagWidget.cpp
+++ b/Engine/Source/Programs/SubmitTool/Private/View/Widgets/TagWidget.cpp
@@ -64,6 +64,7 @@ void STagWidget::Construct(const FArguments& InArgs)
 	[
 		SNew(SCheckBox)
 		.IsFocusable(false)
+		.IsEnabled_Static(&FModelInterface::GetInputEnabled)
 		.IsChecked_Lambda([&tag = Tag]()->ECheckBoxState {return tag->IsEnabled() ? ECheckBoxState::Checked : ECheckBoxState::Unchecked; })
 		.CheckBoxContentUsesAutoWidth(false)
 		.OnCheckStateChanged(this, &STagWidget::OnCheckboxChangedEvent)
@@ -76,6 +77,7 @@ void STagWidget::Construct(const FArguments& InArgs)
 	.VAlign(VAlign_Center)
 	[
 		SNew(SButton)
+		.IsEnabled_Static(&FModelInterface::GetInputEnabled)
 		.ButtonStyle(FAppStyle::Get(), "InvisibleButton")
 		.IsFocusable(false)
 		.OnClicked_Raw(this, &STagWidget::OnLabelClick)
@@ -228,7 +230,16 @@ void STagWidget::Construct(const FArguments& InArgs)
 							{
 								return FText::FromString(TEXT("Request a review using the current users as reviewers."));
 							}})
-						.IsEnabled_Lambda([&ModelInterface = ModelInterface]() { return ModelInterface->IsSwarmServiceValid() && !ModelInterface->IsP4OperationRunning(); })
+						.IsEnabled_Lambda([&ModelInterface = ModelInterface]() { 
+							if (ModelInterface->HasSwarmReview())
+							{
+								return ModelInterface->GetInputEnabled() && ModelInterface->IsSwarmServiceValid() && !ModelInterface->GetSwarmService().Pin()->IsRequestRunning();
+							}
+							else
+							{
+								return ModelInterface->IsSwarmServiceValid() && !ModelInterface->GetSwarmService().Pin()->IsRequestRunning();
+							}
+						})
 						[									
 							SNew(STextBlock)
 							.Justification(ETextJustify::Center)
@@ -300,7 +311,7 @@ void STagWidget::Construct(const FArguments& InArgs)
 				.MinDesiredWidth(60)
 				[
 					SNew(SComboButton)
-					.IsEnabled(this->SelectValues.Num() > 0)
+					.IsEnabled_Lambda([this](){ return FModelInterface::GetInputEnabled() && SelectValues.Num() > 0;})
 					.ButtonContent()
 					[
 						SNew(STextBlock)
@@ -327,6 +338,7 @@ void STagWidget::Construct(const FArguments& InArgs)
 					[
 						SNew(SCheckBox)
 						.IsChecked_Lambda([this, i]() { return this->Tag->GetValues().Contains(*this->SelectValues[i]) ? ECheckBoxState::Checked : ECheckBoxState::Unchecked; })
+						.IsEnabled_Static(&FModelInterface::GetInputEnabled)
 						.OnCheckStateChanged_Lambda([this, i](ECheckBoxState InNewState)
 						{
 							 OnSelectedChangedFromMultiselect(SelectValues[i]);
@@ -338,8 +350,8 @@ void STagWidget::Construct(const FArguments& InArgs)
 					[
 						SNew(SButton)
 						.ButtonStyle(FAppStyle::Get(), "InvisibleButton")
-						.IsFocusable(false)			
-
+						.IsFocusable(false)
+						.IsEnabled_Static(&FModelInterface::GetInputEnabled)
 						.OnClicked_Lambda([this, i]() { OnSelectedChangedFromMultiselect(SelectValues[i]); return FReply::Handled(); })
 						[
 							SNew(STextBlock)
@@ -364,6 +376,7 @@ void STagWidget::Construct(const FArguments& InArgs)
 				.MinDesiredWidth(60)
 				[
 					SNew(SButton)
+					.IsEnabled_Static(&FModelInterface::GetInputEnabled)
 					.OnClicked(this, &STagWidget::OnJiraClick)
 					.HAlign(HAlign_Center)
 					.VAlign(VAlign_Center)
@@ -421,7 +434,7 @@ void STagWidget::Construct(const FArguments& InArgs)
 				.MinDesiredWidth(60)
 				[
 					SNew(SComboButton)
-					.IsEnabled_Lambda([this](){ return PreflightListUI->Num() > 0; })
+					.IsEnabled_Lambda([this](){ return ModelInterface->GetInputEnabled() && PreflightListUI->Num() > 0; })
 					.ForegroundColor_Raw(this, &STagWidget::GetPreflightGlobalColor)
 					.ButtonContent()
 					[
@@ -446,7 +459,7 @@ void STagWidget::Construct(const FArguments& InArgs)
 				[
 					SNew(SButton)
 					.ContentPadding(-5)
-					.IsEnabled_Lambda([this]{ return FModelInterface::GetInputEnabled() && Tag->GetValues().Num() > 0; })
+					.IsEnabled_Lambda([this]{ return Tag->GetValues().Num() > 0; })
 					.IsFocusable(true)
 					.ToolTipText(FText::FromString(TEXT("Open in browser")))
 					.OnClicked_Lambda([this]()
@@ -774,4 +787,4 @@ FReply STagWidget::OnSwarmClick()
 
 	return FReply::Handled();
 }
-#undef LOCTEXT_NAMESPACE
\ No newline at end of file
+#undef LOCTEXT_NAMESPACE
