[SubmitTool] Make Preflight template optional when using validator MinMemory + allow skipping when not enough memory by config

#Synced-CL 48728512

diff --git a/Engine/Source/Developer/SubmitToolCore/Public/Logic/Validators/ValidatorDefinition.h b/Engine/Source/Developer/SubmitToolCore/Public/Logic/Validators/ValidatorDefinition.h
--- a/Engine/Source/Developer/SubmitToolCore/Public/Logic/Validators/ValidatorDefinition.h
+++ b/Engine/Source/Developer/SubmitToolCore/Public/Logic/Validators/ValidatorDefinition.h
@@ -267,6 +267,9 @@ struct FValidatorDefinition
 	UPROPERTY()
 	FString NotEnoughMemoryMessage = TEXT("There is not enough available memory for the estimated consumption. Consider closing other programs and retrying.");
 
+	UPROPERTY()
+	bool bAllowSkippingOnMemoryFail = false;
+
 	UPROPERTY()
 	FString PreflightTemplate;
 };
diff --git a/Engine/Source/Programs/SubmitTool/Private/Logic/Validators/ValidatorBase.cpp b/Engine/Source/Programs/SubmitTool/Private/Logic/Validators/ValidatorBase.cpp
--- a/Engine/Source/Programs/SubmitTool/Private/Logic/Validators/ValidatorBase.cpp
+++ b/Engine/Source/Programs/SubmitTool/Private/Logic/Validators/ValidatorBase.cpp
@@ -137,14 +137,29 @@ void FValidatorBase::StartValidation()
 					TStringBuilder<256> Message;
 					Message.Append(Definition->NotEnoughMemoryMessage);
 					Message.Appendf(TEXT("\nAvailable memory: %llu GB\nRequired memory: %d GB"), MemoryStats.AvailablePhysical / 1024 / 1024 / 1024, Definition->MinMemoryGBRequired);
-					const FString FinalMessage = FString::Format(Message.ToString(), PreflightService->GetFormatParameters(Definition->PreflightTemplate));
+					const FString FinalMessage = Definition->PreflightTemplate.IsEmpty() ? Message.ToString() : FString::Format(Message.ToString(), PreflightService->GetFormatParameters(Definition->PreflightTemplate));
 
-					TArray<FString> Buttons{ TEXT("Request Preflight"), TEXT("Try anyway"), TEXT("Close") };
-					EDialogFactoryResult Result = FDialogFactory::ShowDialog(FText::FromString(GetValidatorName() + TEXT(" - Not Enough Memory")), FText::FromString(FinalMessage), Buttons, nullptr, [this](size_t InIdx) { return InIdx != 0 || !Definition->PreflightTemplate.IsEmpty(); });
-					if (Result != EDialogFactoryResult::SecondButton)
+					TArray<FString> Buttons{ TEXT("Try anyway"), Definition->bAllowSkippingOnMemoryFail ? TEXT("Skip Validation") : TEXT("Close")};
+
+					if(!Definition->PreflightTemplate.IsEmpty())
 					{
+						Buttons.Add(TEXT("Request Preflight"));
+					}
+
+					EDialogFactoryResult Result = FDialogFactory::ShowDialog(FText::FromString(GetValidatorName() + TEXT(" - Not Enough Memory")), FText::FromString(FinalMessage), Buttons, nullptr);
+					
+					if (Result != EDialogFactoryResult::FirstButton)
+					{						
+						if(Definition->bAllowSkippingOnMemoryFail && Result == EDialogFactoryResult::SecondButton)
+						{
+							FailureReason = TEXT("Out of Memory");
+							UE_LOG(LogSubmitTool, Log, TEXT("[%s] Skipped due to lack of memory.\n%s"), *GetValidatorName(), *FinalMessage);
+							Skip();
+							return;
+						}
+
 						LogFailure(FString::Printf(TEXT("[%s] %s"), *GetValidatorName(), *FinalMessage));
-						if (Result == EDialogFactoryResult::FirstButton)
+						if (Result == EDialogFactoryResult::ThirdButton)
 						{
 							PreflightService->RequestPreflight(Definition->PreflightTemplate);
 						}
diff --git a/Engine/Source/Programs/SubmitTool/Private/Version/AppVersionDefines.h b/Engine/Source/Programs/SubmitTool/Private/Version/AppVersionDefines.h
--- a/Engine/Source/Programs/SubmitTool/Private/Version/AppVersionDefines.h
+++ b/Engine/Source/Programs/SubmitTool/Private/Version/AppVersionDefines.h
@@ -21,7 +21,7 @@
 
 #define SUBMIT_TOOL_VERSION_MAJOR 6
 #define SUBMIT_TOOL_VERSION_MINOR 5
-#define SUBMIT_TOOL_VERSION_PATCH 3
+#define SUBMIT_TOOL_VERSION_PATCH 4
 
 #ifndef SUBMIT_TOOL_CHANGELIST
 #define SUBMIT_TOOL_CHANGELIST_STRING TEXT("NoCL")
