[Submit Tool] Fixed inverted logic for HasSwarmReview

#Synced-CL 48728512

diff --git a/Engine/Source/Programs/SubmitTool/Private/View/Widgets/TagWidget.cpp b/Engine/Source/Programs/SubmitTool/Private/View/Widgets/TagWidget.cpp
--- a/Engine/Source/Programs/SubmitTool/Private/View/Widgets/TagWidget.cpp
+++ b/Engine/Source/Programs/SubmitTool/Private/View/Widgets/TagWidget.cpp
@@ -231,7 +231,7 @@ void STagWidget::Construct(const FArguments& InArgs)
 								return FText::FromString(TEXT("Request a review using the current users as reviewers."));
 							}})
 						.IsEnabled_Lambda([&ModelInterface = ModelInterface]() { 
-							if (ModelInterface->HasSwarmReview())
+							if (!ModelInterface->HasSwarmReview())
 							{
 								return ModelInterface->GetInputEnabled() && ModelInterface->IsSwarmServiceValid() && !ModelInterface->GetSwarmService().Pin()->IsRequestRunning();
 							}
