[SubmitTool] Renaming OOM popup button to Retry

#Synced-CL 48728512

diff --git a/Engine/Source/Programs/SubmitTool/Private/Logic/Validators/ValidatorBase.cpp b/Engine/Source/Programs/SubmitTool/Private/Logic/Validators/ValidatorBase.cpp
--- a/Engine/Source/Programs/SubmitTool/Private/Logic/Validators/ValidatorBase.cpp
+++ b/Engine/Source/Programs/SubmitTool/Private/Logic/Validators/ValidatorBase.cpp
@@ -139,7 +139,7 @@ void FValidatorBase::StartValidation()
 					Message.Appendf(TEXT("\nAvailable memory: %llu GB\nRequired memory: %d GB"), MemoryStats.AvailablePhysical / 1024 / 1024 / 1024, Definition->MinMemoryGBRequired);
 					const FString FinalMessage = Definition->PreflightTemplate.IsEmpty() ? Message.ToString() : FString::Format(Message.ToString(), PreflightService->GetFormatParameters(Definition->PreflightTemplate));
 
-					TArray<FString> Buttons{ TEXT("Try anyway"), Definition->bAllowSkippingOnMemoryFail ? TEXT("Skip Validation") : TEXT("Close")};
+					TArray<FString> Buttons{ TEXT("Retry"), Definition->bAllowSkippingOnMemoryFail ? TEXT("Skip Validation") : TEXT("Close")};
 
 					if(!Definition->PreflightTemplate.IsEmpty())
 					{
