[SubmitTool] Adding FromEditor property to validators

#Synced-CL 48728512

diff --git a/Engine/Source/Programs/SubmitTool/Private/Logic/Validators/ValidatorBase.cpp b/Engine/Source/Programs/SubmitTool/Private/Logic/Validators/ValidatorBase.cpp
--- a/Engine/Source/Programs/SubmitTool/Private/Logic/Validators/ValidatorBase.cpp
+++ b/Engine/Source/Programs/SubmitTool/Private/Logic/Validators/ValidatorBase.cpp
@@ -152,7 +152,7 @@ void FValidatorBase::StartValidation()
 					{						
 						if(Definition->bAllowSkippingOnMemoryFail && Result == EDialogFactoryResult::SecondButton)
 						{
-							FailureReason = TEXT("Out of Memory");
+							FailureReason = TEXT("Skipped: Out of Memory");
 							UE_LOG(LogSubmitTool, Log, TEXT("[%s] Skipped due to lack of memory.\n%s"), *GetValidatorName(), *FinalMessage);
 							Skip();
 							return;
@@ -355,6 +355,7 @@ const TArray<FAnalyticsEventAttribute> FValidatorBase::GetTelemetryAttributes()
 	const TSharedPtr<FChangelistService> ChangelistService = ServiceProvider.Pin()->GetService<FChangelistService>();
 	FString ChangelistNumber = ChangelistService->GetCLID();
 
+
 	return MakeAnalyticsEventAttributeArray(
 		TEXT("ValidatorID"), *GetValidatorNameId().ToString(),
 		TEXT("ValidatorName"), *GetValidatorName(),
@@ -362,7 +363,8 @@ const TArray<FAnalyticsEventAttribute> FValidatorBase::GetTelemetryAttributes()
 		TEXT("Runtime"), RunTime,
 		TEXT("Stream"), ServiceProvider.Pin()->GetService<ISTSourceControlService>()->GetCurrentStreamName(),
 		TEXT("PendingChangelist"), MoveTemp(ChangelistNumber),
-		TEXT("FailureReason"), FailureReason
+		TEXT("FailureReason"), FailureReason,
+		TEXT("FromEditor"), FCmdLineParameters::Get().Contains(FSubmitToolCmdLine::EditorFlag)
 	);
 }
 
