SubmitTool: Replace engine source control changelist handling

- Removed unused param-file from cmdline
- Moved ModelInterface initialization off the main thread
- Data drived from config syncing files on startup
- Moved all p4 commands away from the Engine's source control module
- Simplified some p4 commands (merged revert unchanged with submit)
- Removed PerforceSourceControl Plugin
- Fixed SubmitToolPerforce cancel request system
- Keep track of failed p4 connections
- Added support for passing data to p4 through stdin in SubmitToolPerforce
- Adds p4 icons to the filelist
- Small fixes

#Synced-CL 49723226

diff --git a/Engine/Source/Developer/SubmitToolCore/Private/Models/Tag.cpp b/Engine/Source/Developer/SubmitToolCore/Private/Models/Tag.cpp
--- a/Engine/Source/Developer/SubmitToolCore/Private/Models/Tag.cpp
+++ b/Engine/Source/Developer/SubmitToolCore/Private/Models/Tag.cpp
@@ -153,16 +153,16 @@ const TArray<FString> FTag::GetValues(bool bEvenIfDisabled) const
 	}
 }
 
-const FTagValidationConfig& FTag::GetCurrentValidationConfig(const TArray<FString>& InDepotPaths) const
+const FTagValidationConfig& FTag::GetCurrentValidationConfig(const TArray<FSCFileRef>& InFiles) const
 {
 	if(Definition.ValidationOverrides.Num() != 0)
 	{
 		for(const FTagValidationOverride& ValidationOverride : Definition.ValidationOverrides)
 		{
 			FRegexPattern RegexPattern = FRegexPattern(ValidationOverride.RegexPath, ERegexPatternFlags::CaseInsensitive);
-			for(const FString& Path : InDepotPaths)
+			for(const FSCFileRef& File : InFiles)
 			{
-				FRegexMatcher Regex = FRegexMatcher(RegexPattern, Path);
+				FRegexMatcher Regex = FRegexMatcher(RegexPattern, File->GetDepotPath());
 				if(Regex.FindNext())
 				{
 					return ValidationOverride.ConfigOverride;
diff --git a/Engine/Source/Developer/SubmitToolCore/Public/CommandLine/CmdLineParameters.h b/Engine/Source/Developer/SubmitToolCore/Public/CommandLine/CmdLineParameters.h
--- a/Engine/Source/Developer/SubmitToolCore/Public/CommandLine/CmdLineParameters.h
+++ b/Engine/Source/Developer/SubmitToolCore/Public/CommandLine/CmdLineParameters.h
@@ -13,7 +13,6 @@ namespace FSubmitToolCmdLine
 	const FString P4User = TEXT("user");
 	const FString P4ChangeList = TEXT("cl");
 	const FString RootDir = TEXT("root-dir");
-	const FString ParameterFile = TEXT("param-file");
 	const FString EditorFlag = TEXT("from-editor");
 
 	const TArray<TSharedPtr<FCmdLineParameter>> SubmitToolCmdLineArgs =
@@ -50,7 +49,7 @@ namespace FSubmitToolCmdLine
 			OutValue.MidInline(Left, (Right - Left) + 1);
 
 			#if PLATFORM_WINDOWS
-			if (OutValue[1] == ':')
+			if (OutValue.Len() > 1 && OutValue[1] == ':')
 			{
 				OutValue[0] = FChar::ToUpper(OutValue[0]);
 			}
@@ -58,8 +57,7 @@ namespace FSubmitToolCmdLine
 			
 			FPaths::NormalizeDirectoryName(OutValue);
 		}),
-		MakeShared<FCmdLineParameter>(ParameterFile, false, TEXT("Submit tool parameters XML override.")),
-		MakeShared<FCmdLineParameter>(EditorFlag, false, TEXT("Submit tool parameters XML override."), true),
+		MakeShared<FCmdLineParameter>(EditorFlag, false, TEXT("Flag indicating if submit tool is being invoked from unreal editor."), true),
 	};
 }
 
diff --git a/Engine/Source/Developer/SubmitToolCore/Public/Models/SCFile.h b/Engine/Source/Developer/SubmitToolCore/Public/Models/SCFile.h
new file mode 100644
--- /dev/null
+++ b/Engine/Source/Developer/SubmitToolCore/Public/Models/SCFile.h
@@ -0,0 +1,45 @@
+// Copyright Epic Games, Inc. All Rights Reserved.
+
+#pragma once
+
+#include "Templates/SharedPointer.h"
+#include "SCFile.generated.h"
+
+UENUM()
+enum ESCFileState : int64
+{
+	Edit,
+	Add,
+	Delete,
+	Integrate,
+	Unknown
+};
+
+class FSCFile final : public TSharedFromThis<class FSCFile, ESPMode::ThreadSafe>
+{
+public:
+	FSCFile(const FString& InFilename, const FString& InDepotPath, ESCFileState InState)
+	: Filename(InFilename)
+	, DepotPath(InDepotPath)
+	, State(InState)
+	{ }
+	~FSCFile() {}
+
+	const FString& GetFilename() const { return Filename; }
+	const FString& GetDepotPath() const { return DepotPath; }
+	ESCFileState GetState() const { return State; }
+	void UpdateState(ESCFileState NewState) { State = NewState; }
+	void UpdateFilename(const FString& NewFilename) { Filename = NewFilename; }
+
+	bool IsAdded() const { return State == ESCFileState::Add; };
+	bool IsCheckedOut() const { return State == ESCFileState::Edit || State == ESCFileState::Unknown; }
+	bool IsIntegrate() const { return State == ESCFileState::Integrate; }	
+	bool IsDeleted() const { return State == ESCFileState::Delete; }
+
+private:
+	FString Filename;
+	FString DepotPath;
+	ESCFileState State;
+};
+
+using FSCFileRef = TSharedRef<FSCFile, ESPMode::ThreadSafe>;
diff --git a/Engine/Source/Developer/SubmitToolCore/Public/Models/Tag.h b/Engine/Source/Developer/SubmitToolCore/Public/Models/Tag.h
--- a/Engine/Source/Developer/SubmitToolCore/Public/Models/Tag.h
+++ b/Engine/Source/Developer/SubmitToolCore/Public/Models/Tag.h
@@ -3,6 +3,7 @@
 #pragma once
 
 #include "TagDefinition.h"
+#include "SCFile.h"
 
 class FTag;
 
@@ -34,7 +35,7 @@ public:
 	void SetValues(const TArray<FString>& InValues); 
 	const TArray<FString> GetValues(bool bIncludeInactive = false) const;
 
-	const FTagValidationConfig& GetCurrentValidationConfig(const TArray<FString>& InDepotPaths) const;
+	const FTagValidationConfig& GetCurrentValidationConfig(const TArray<FSCFileRef>& InDepotPaths) const;
 
 	const ETagState GetTagState() const
 	{
diff --git a/Engine/Source/Developer/SubmitToolCore/Public/Parameters/SubmitToolParameters.h b/Engine/Source/Developer/SubmitToolCore/Public/Parameters/SubmitToolParameters.h
--- a/Engine/Source/Developer/SubmitToolCore/Public/Parameters/SubmitToolParameters.h
+++ b/Engine/Source/Developer/SubmitToolCore/Public/Parameters/SubmitToolParameters.h
@@ -28,12 +28,12 @@ struct FGeneralParameters
 	GENERATED_BODY()
 
 	UPROPERTY()
-	FString NewChangelistMessage = TEXT("Submit Tool generated changelist from default with {FileCount} files");
+	FString NewChangelistMessage = TEXT("Submit Tool generated changelist from default");
 
 	UPROPERTY()
 	TArray<FString> ForbiddenDescriptions = 
 	{
-		TEXT("Submit Tool generated changelist from default with"),
+		TEXT("Submit Tool generated changelist from default"),
 		TEXT("<saved by perforce>")
 	};
 
@@ -66,6 +66,9 @@ struct FGeneralParameters
 	
 	UPROPERTY()
 	TArray<FString> AutomaticChangelistTags;
+
+	UPROPERTY()
+	TArray<FString> SyncOnStartupPaths;
 };
 
 USTRUCT()
diff --git a/Engine/Source/Programs/SubmitTool/Private/Logic/ChangelistService.cpp b/Engine/Source/Programs/SubmitTool/Private/Logic/ChangelistService.cpp
--- a/Engine/Source/Programs/SubmitTool/Private/Logic/ChangelistService.cpp
+++ b/Engine/Source/Programs/SubmitTool/Private/Logic/ChangelistService.cpp
@@ -2,6 +2,7 @@
 
 #include "ChangelistService.h"
 
+#include "Algo/Compare.h"
 #include "Algo/Find.h"
 #include "Configuration/Configuration.h"
 #include "Logging/SubmitToolLog.h"
@@ -12,15 +13,7 @@
 #include "HAL/FileManagerGeneric.h"
 #include "CommandLine/CmdLineParameters.h"
 #include "SubmitToolUtils.h"
-
-#include "ISourceControlProvider.h"
-#include "ISourceControlOperation.h"
-#include "SourceControlOperations.h"
-#include "SourceControlResultInfo.h"
-
-namespace ChangelistServiceConstants {
-	constexpr float TickDelay = 5.f;
-}
+#include "SubmitToolCoreUtils.h"
 
 FChangelistService::FChangelistService(
 	const FGeneralParameters& InParameters,
@@ -30,7 +23,7 @@ FChangelistService::FChangelistService(
 	Parameters(InParameters),
 	CLReadyCallback(InCLReadyCallback),
 	CLRefreshCallback(InCLRefreshCallback),
-	TickHandle(FTSTicker::GetCoreTicker().AddTicker(FTickerDelegate::CreateRaw(this, &FChangelistService::P4Tick), ChangelistServiceConstants::TickDelay)),
+	TickHandle(FTSTicker::GetCoreTicker().AddTicker(FTickerDelegate::CreateRaw(this, &FChangelistService::P4Tick))),
 	ServiceProvider(InServiceProvider),
 	SourceControlService(InServiceProvider->GetService<ISTSourceControlService>())
 {
@@ -46,21 +39,15 @@ FChangelistService::~FChangelistService()
 void FChangelistService::Init()
 {
 	TRACE_CPUPROFILER_EVENT_SCOPE(FChangelistService::Init);
-	if(SourceControlService->GetProvider().IsValid())
+	if (SourceControlService.IsValid())
 	{
-		InitialP4OpenedFilesTask = LaunchGetOpenedFilesDepotPaths();
-
-		CachedSCCProvider = SourceControlService->GetProvider().Get();
-
-		if(!CLID.IsNumeric() && CLID.Equals(TEXT("default"), ESearchCase::IgnoreCase))
+		if (!CLID.IsNumeric() && CLID.Equals(TEXT("default"), ESearchCase::IgnoreCase))
 		{
 			CreateCLFromDefaultCL();
 		}
 		else
 		{
-			// Only do a p4 describe if we're not using the default CL
-			InitialP4ShelvedFilesTask = LaunchGetShelvedFilesDepotPaths();
-			FindInitialChangelistsAsync();
+			FetchChangelistDataAsync();
 		}
 	}
 	else
@@ -69,429 +56,501 @@ void FChangelistService::Init()
 	}
 }
 
-void FChangelistService::FindInitialChangelistsAsync()
+void FChangelistService::FetchChangelistDataAsync()
 {
-	TRACE_CPUPROFILER_EVENT_SCOPE(FChangelistService::FindInitialChangelistsAsync);
+	TRACE_CPUPROFILER_EVENT_SCOPE(FChangelistService::FetchChangelistDataAsync);
+	UE_LOG(LogSubmitToolP4Debug, Log, TEXT("Fetching changelist"));
+
+	// Wait for any already running tasks
+	WaitForChangelistData();
 
-	TSharedRef<FUpdatePendingChangelistsStatus> UpdatePendingChangelistsOperation = ISourceControlOperation::Create<FUpdatePendingChangelistsStatus>();
-	UpdatePendingChangelistsOperation->SetUpdateAllChangelists(true);
+	// First, launch two "opened" and "describe" commands in parallel
+	OpenedTask = SourceControlService->RunCommand(TEXT("opened"), { });
+	DescribeTask = SourceControlService->RunCommand(TEXT("describe"), { "-S", CLID });
 
-	UE_LOG(LogSubmitToolP4Debug, Log, TEXT("Finding available changelists"));
+	// Then, wait for both before running an additional "where" command which
+	// is required to build a complete representation of the changelist state.
+	ParseChangelistTask = UE::Tasks::Launch(UE_SOURCE_LOCATION, [this]
+	{
+		TRACE_CPUPROFILER_EVENT_SCOPE(FChangelistService::ParseChangelistTask);
 
-	ActiveP4Operations.Add(UpdatePendingChangelistsOperation);
+		check(OpenedTask.IsCompleted())
+		check(DescribeTask.IsCompleted());
 
-	CachedSCCProvider->Execute(UpdatePendingChangelistsOperation, EConcurrency::Asynchronous, FSourceControlOperationComplete::CreateLambda([this](const FSourceControlOperationRef& UpdateOperation, ECommandResult::Type Result)
+		TSharedRef<FSCCommand> OpenedCmd = OpenedTask.GetResult();
+		TSharedRef<FSCCommand> DescribeCmd = DescribeTask.GetResult();
+
+		if (!OpenedCmd->bSuccess || !DescribeCmd->bSuccess)
 		{
-			TRACE_CPUPROFILER_EVENT_SCOPE(FChangelistService::FindInitialChangelistsAsync_Callback);
-			if(Result == ECommandResult::Succeeded)
-			{
-				TArray<FSourceControlChangelistRef> Changelists = CachedSCCProvider->GetChangelists(EStateCacheUsage::Use);
-				if (FSourceControlChangelistRef* Changelist = Algo::FindByPredicate(Changelists, [this](const FSourceControlChangelistRef& CL){ return CL->GetIdentifier() == CLID; }))
-				{
-					ChangelistPtr = *Changelist;
-				}
-			}
-			else if(Result == ECommandResult::Cancelled)
+			return FChangelistData{};
+		}
+
+		// Do a first quick parse of the fetched data in order to
+		// collect the set of open and shelved depot paths in the changelist.
+		TSet<FString> DepotFiles;
+		for (const TMap<FString, FString>& Record : OpenedCmd->Values)
+		{
+			const FString* DepotFile = Record.Find(TEXT("depotFile"));
+			const FString* Change = Record.Find(TEXT("change"));
+			if (DepotFile && Change && *Change == CLID)
 			{
-				UE_LOG(LogSubmitToolP4, Warning, TEXT("Cancelled finding available changelists"));
+				DepotFiles.Add(*DepotFile);
 			}
-			else if(Result == ECommandResult::Failed)
+		}
+		for (const TMap<FString, FString>& Record : DescribeCmd->Values)
+		{
+			int32 I = 0;
+			while (const FString* DepotFile = Record.Find(FString::Printf(TEXT("depotFile%d"), I)))
 			{
-				UE_LOG(LogSubmitToolP4, Error, TEXT("Failed to find available changelists."));
-				PrintErrorMessages(UpdateOperation->GetResultInfo());
+				DepotFiles.Add(*DepotFile);
+				++I;
 			}
+		}
+		
+		// Run a sync "where" query inside this task using the collected depot paths to 
+		// build a mapping to local filenames
+		TMap<FString, FString> DepotFileToLocalFile = RunWhereCommandSync(DepotFiles.Array());
 
-			ActiveP4Operations.Remove(UpdateOperation);
+		return FChangelistData::Parse(CLID, *DescribeCmd, *OpenedCmd, DepotFileToLocalFile);
+	}, Prerequisites(OpenedTask, DescribeTask));
+
+	// Last, launch a GT task to update the actual model data
+	UpdateChangelistTask = UE::Tasks::Launch(UE_SOURCE_LOCATION, [this]
+	{
+		UE_LOG(LogSubmitToolP4Debug, Log, TEXT("Fetch changelist operation succeeded."));
+
+		if (Changelist.bIsValid)
+		{
+			UpdateCachedChangelistData();
+		}
+		else
+		{
+			CacheChangelistData();
+		}
+	}, ParseChangelistTask, LowLevelTasks::ETaskPriority::Normal, UE::Tasks::EExtendedTaskPriority::GameThreadNormalPri);
+}
 
-			if(ChangelistPtr)
+void FChangelistService::RevertUnchangedFiles(TFunction<void(const FSCCommand&)> OnComplete)
+{
+	TRACE_CPUPROFILER_EVENT_SCOPE(FChangelistService::RevertUnchangedFiles);
+	UE_LOG(LogSubmitToolP4, Log, TEXT("Reverting unchanged files from CL %s..."), *GetCLID());
+	
+	SourceControlService->RunCommand(TEXT("revert"), { TEXT("-a") , TEXT("-c"), GetCLID() }, FOnSCCCommandComplete::CreateLambda(
+		[OnComplete, this](const FSCCommand& Cmd)
+		{
+			TRACE_CPUPROFILER_EVENT_SCOPE(FChangelistService::RevertUnchangedFiles_Callback);
+			if (Cmd.bSuccess)
 			{
-				FetchChangelistDataAsync();
+				for (const FSCCRecord& Record : Cmd.Values)
+				{
+				}
+				
+				UE_LOG(LogSubmitToolP4, Log, TEXT("CL %s Revert unchanged operation succeeded"), *GetCLID());
 			}
 			else
 			{
-				UE_LOG(LogSubmitToolP4, Error, TEXT("Couldn't retrieve information for CL %s"), *CLID);
-				FModelInterface::SetErrorState();
+				UE_LOG(LogSubmitToolP4, Error, TEXT("CL %s Revert unchanged operation failed:"), *GetCLID());
 			}
-		}));
-}
-
-void FChangelistService::FetchChangelistDataAsync()
-{
-	check(ChangelistPtr.IsValid());
-
-	TRACE_CPUPROFILER_EVENT_SCOPE(FChangelistService::FetchChangelistDataAsync);
 
-	UE_LOG(LogSubmitToolP4Debug, Log, TEXT("Updating CL %s changes from P4"), *GetCLID());
-	TSharedRef<FUpdatePendingChangelistsStatus> UpdatePendingChangelistsOperation = ISourceControlOperation::Create<FUpdatePendingChangelistsStatus>();
-	UpdatePendingChangelistsOperation->SetUpdateFilesStates(true);
-	UpdatePendingChangelistsOperation->SetUpdateShelvedFilesStates(true);
-	FSourceControlChangelistRef Changelist = ChangelistPtr.ToSharedRef();
-	UpdatePendingChangelistsOperation.Get().SetChangelistsToUpdate(MakeArrayView(&Changelist, 1));
-
-	ActiveP4Operations.Add(UpdatePendingChangelistsOperation);
-	CachedSCCProvider->Execute(UpdatePendingChangelistsOperation, EConcurrency::Asynchronous, FSourceControlOperationComplete::CreateRaw(this, &FChangelistService::FetchChangelistCallback));
+			if (OnComplete.IsSet())
+			{
+				OnComplete(Cmd);
+			}
+		}
+	));
 }
 
 
-void FChangelistService::FetchChangelistCallback(const FSourceControlOperationRef& UpdateOperation, ECommandResult::Type Result)
+void FChangelistService::UpdateP4CLDescriptionSynchronously()
 {
-	TRACE_CPUPROFILER_EVENT_SCOPE(FChangelistService::FetchChangelistCallback);
-
-	if(Result == ECommandResult::Succeeded)
+	TRACE_CPUPROFILER_EVENT_SCOPE(FChangelistService::UpdateP4CLDescriptionSynchronously);
+	if (!CLDescription.Equals(Changelist.Description))
 	{
-		UE_LOG(LogSubmitToolP4Debug, Log, TEXT("Fetch changelist operation succeeded."));
-		PrintMessages(UpdateOperation->GetResultInfo());
+		UE_LOG(LogSubmitToolP4Debug, Log, TEXT("Saving CL description %s"), *GetCLID());
+		FSCCommand OutputCmd = SourceControlService->RunSyncCommand(TEXT("change"), { TEXT("-o"), GetCLID() });
 
-		TArray<FSourceControlChangelistRef> Changelists = CachedSCCProvider->GetChangelists(EStateCacheUsage::Use);
-		if (Algo::FindByPredicate(Changelists, [this](const FSourceControlChangelistRef& CL){ return CL->GetIdentifier() == CLID; }))
+		if (OutputCmd.bSuccess && OutputCmd.Values.Num() > 0)
 		{
-			// If this is the first time we fetch results from perforce, we cache our CL State object and let the app know 
-			if (!ChangelistState.IsValid())
+			FSCCommand ChangeCmd = FSCCommand(TEXT("change"), { TEXT("-i") });
+			ChangeCmd.InputData = ParseP4ChangesOutput(OutputCmd.Values[0], CLDescription);
+			SourceControlService->RunSyncCommand(ChangeCmd);
+
+			if (ChangeCmd.bSuccess)
 			{
-				ChangelistState = CachedSCCProvider->GetState(ChangelistPtr.ToSharedRef(), EStateCacheUsage::Use);
-				CacheChangelistState();
-				PrintFilesAndShelvedFiles();
-				CLReadyCallback.ExecuteIfBound(ChangelistState.IsValid());
+				UE_LOG(LogSubmitToolP4, Log, TEXT("CL %s edit changelist operation succeeded."), *GetCLID());
+				Changelist.Description = CLDescription;
 			}
-			// else we simply update our cached stae
 			else
 			{
-				ChangelistState = CachedSCCProvider->GetState(ChangelistPtr.ToSharedRef(), EStateCacheUsage::Use);
-				UpdateCachedChangelistState();
+				UE_LOG(LogSubmitToolP4, Error, TEXT("CL %s edit changelist operation failed"), *GetCLID());
 			}
 		}
-		else
-		{
-			UE_LOG(LogSubmitToolP4, Error, TEXT("Failed to fetch and update CL %s from P4"), *CLID);
-		}
 	}
-	else if(Result == ECommandResult::Cancelled)
+}
+
+void FChangelistService::UpdateP4CLDescription()
+{
+	TRACE_CPUPROFILER_EVENT_SCOPE(FChangelistService::UpdateP4CLDescription);
+
+	if (!CLDescription.Equals(Changelist.Description))
 	{
-		UE_LOG(LogSubmitToolP4, Warning, TEXT("Update operation cancelled."));
+		UE_LOG(LogSubmitToolP4Debug, Log, TEXT("Saving CL description %s"), *GetCLID());
+		
+		SourceControlService->RunCommand(TEXT("change"), { TEXT("-o"), GetCLID() }, FOnSCCCommandComplete::CreateLambda([this](const FSCCommand& Cmd) {
+			TRACE_CPUPROFILER_EVENT_SCOPE(FChangelistService::UpdateP4CLDescription_change_output_callback);
+			if (Cmd.bSuccess && Cmd.Values.Num() > 0)
+			{
+				FSCCommand ChangeCmd = FSCCommand(TEXT("change"), { TEXT("-i") });
+				ChangeCmd.InputData = ParseP4ChangesOutput(Cmd.Values[0], CLDescription);
+				SourceControlService->RunCommand(MoveTemp(ChangeCmd), FOnSCCCommandComplete::CreateLambda([this](const FSCCommand& Cmd) {
+					TRACE_CPUPROFILER_EVENT_SCOPE(FChangelistService::UpdateP4CLDescription_change_input_callback);
+					if (Cmd.bSuccess)
+					{
+						UE_LOG(LogSubmitToolP4, Log, TEXT("CL %s edit changelist operation succeeded."), *GetCLID());
+						Changelist.Description = CLDescription;
+					}
+					else
+					{
+						UE_LOG(LogSubmitToolP4, Error, TEXT("CL %s edit changelist operation failed"), *GetCLID());
+					}
+				}));
+			}
+			else
+			{
+				UE_LOG(LogSubmitToolP4, Error, TEXT("CL %s edit changelist operation failed"), *GetCLID());
+
+			}
+		}));		
 	}
-	else if(Result == ECommandResult::Failed)
+}
+
+static ESCFileState GetLocalFileState(const FString& State)
+{
+	if (State == TEXT("edit"))
+	{
+		return ESCFileState::Edit;
+	}
+	if (State == TEXT("openforadd") || State == TEXT("move/add") || State == TEXT("add"))
+	{
+		return ESCFileState::Add;
+	}
+	if (State == TEXT("markfordelete") || State == TEXT("move/delete") || State == TEXT("delete"))
+	{
+		return ESCFileState::Delete;
+	}
+	if (State == TEXT("integrate") || State == TEXT("branch"))
 	{
-		UE_LOG(LogSubmitToolP4, Warning, TEXT("Update operation failed."));
-		PrintErrorMessages(UpdateOperation->GetResultInfo());
+		return ESCFileState::Integrate;
 	}
 
-	ActiveP4Operations.Remove(UpdateOperation);
+	// Unknown State
+	ensureMsgf(false, TEXT("GetLocalFileState: Unknown perforce file state: '%s'"), *State);
+	return ESCFileState::Unknown;
 }
 
-void FChangelistService::RevertUnchangedFilesAsync(const FSourceControlOperationComplete& OnRevertComplete)
+FChangelistData FChangelistData::Parse(
+	const FString& CLID,
+	const FSCCommand& DescribeCmd,
+	const FSCCommand& OpenedCmd,
+	const TMap<FString, FString>& DepotFileToLocalFile)
 {
-	UE_LOG(LogSubmitToolP4, Log, TEXT("Reverting unchanged files from CL %s..."), *GetCLID());
+	TRACE_CPUPROFILER_EVENT_SCOPE(FChangelistData::ParseChangelistCommands);
 
-	TSharedRef<FRevertUnchanged> RevertUnchangedOperation = ISourceControlOperation::Create<FRevertUnchanged>();
+	FChangelistData Out;
 
-	ActiveP4Operations.Add(RevertUnchangedOperation);
+	if (!OpenedCmd.bSuccess || !DescribeCmd.bSuccess)
+	{
+		return Out;
+	}
 
-	CachedSCCProvider->Execute(RevertUnchangedOperation, ChangelistPtr, EConcurrency::Asynchronous, FSourceControlOperationComplete::CreateLambda(
-		[this, OnRevertComplete](const FSourceControlOperationRef& Operation, ECommandResult::Type Result) {
+	bool bFoundDescription = false;
 
-			// Revert Unchanged returns as failed if there were no files to revert, check ErrorMessages to see actual failures
-			if(Result == ECommandResult::Cancelled)
-			{
-				UE_LOG(LogSubmitToolP4, Warning, TEXT("Revert unchanged operation cancelled."));
-			}
-			else if(Result == ECommandResult::Succeeded)
-			{
-				UE_LOG(LogSubmitToolP4, Log, TEXT("Revert unchanged operation succeeded."));
-				FilesInCL = ChangelistState->GetFilesStates();
-				FilesDepotPaths.Reset(FilesInCL.Num());
-				PrintFilesAndShelvedFiles();
-			}
-			else if(Operation->GetResultInfo().ErrorMessages.Num() == 0)
+	// Opened files
+	for (const TMap<FString, FString>& Record : OpenedCmd.Values)
+	{
+		const FString* DepotFile = Record.Find(TEXT("depotFile"));
+		const FString* Change = Record.Find(TEXT("change"));
+		const FString* Action = Record.Find(TEXT("action"));
+		if (DepotFile && Change && Action)
+		{
+			if (*Change == CLID)
 			{
-				UE_LOG(LogSubmitToolP4, Log, TEXT("There were no unchanged files to revert."));
+				Out.Files.Add(MakeShared<FSCFile>(DepotFileToLocalFile.FindRef(*DepotFile), *DepotFile, GetLocalFileState(*Action)));
 			}
-			else if(Result == ECommandResult::Failed)
+			Out.AllOpenFiles.Add(*DepotFile, *Change);
+		}
+	}
+
+	// Describe => Shelved files
+	for (const TMap<FString, FString>& Record : DescribeCmd.Values)
+	{
+		if (const FString* Description = Record.Find(TEXT("desc")))
+		{
+			Out.Description = *Description;
+			bFoundDescription = true;
+		}
+		int32 I = 0;
+		while (const FString* DepotFile = Record.Find(FString::Printf(TEXT("depotFile%d"), I)))
+		{
+			const FString* Action = Record.Find(FString::Printf(TEXT("action%d"), I));
+			if (Action)
 			{
-				UE_LOG(LogSubmitToolP4, Error, TEXT("CL %s revert unchanged operation failed:"), *GetCLID());
-				PrintErrorMessages(Operation->GetResultInfo());
+				Out.ShelvedFiles.Add(MakeShared<FSCFile>(DepotFileToLocalFile.FindRef(*DepotFile), *DepotFile, GetLocalFileState(*Action)));
 			}
-
-			ActiveP4Operations.Remove(Operation);
-			OnRevertComplete.ExecuteIfBound(Operation, Result);
+			++I;
 		}
-	));
+	}
+
+	// Sort and compare open vs. shelved files
+	Algo::StableSortBy(Out.Files, [](const FSCFileRef& File) { return File->GetDepotPath(); });
+	Algo::StableSortBy(Out.ShelvedFiles, [](const FSCFileRef& File) { return File->GetDepotPath(); });
+	Out.bIsShelfComplete = Algo::Compare(Out.Files, Out.ShelvedFiles,
+		[](const FSCFileRef& A, const FSCFileRef& B)
+		{
+			return A->GetDepotPath().Equals(B->GetDepotPath(), ESearchCase::IgnoreCase) && A->GetState() == B->GetState();
+		});
+	
+	// CL is considered valid if the describe command returned a description entry (which could still be empty)
+	Out.bIsValid = bFoundDescription;
+
+	return Out;
 }
 
-void FChangelistService::CacheChangelistState()
+void FChangelistService::CacheChangelistData()
 {
-	check(ChangelistState);
+	TRACE_CPUPROFILER_EVENT_SCOPE(FChangelistService::CacheChangelistData);
+	check(ParseChangelistTask.IsCompleted());
+	Changelist = ParseChangelistTask.GetResult();
+	if (Changelist.bIsValid)
+	{
+		CLDescription = Changelist.Description;
+		PrintFilesAndShelvedFiles();
+	}
+	else
+	{
+		UE_LOG(LogSubmitToolP4, Error, TEXT("Couldn't retrieve information for CL %s"), *CLID);
+		FModelInterface::SetErrorState();
+	}
 
-	OriginalDescription = ChangelistState->GetDescriptionText();
-	CLDescription = ChangelistState->GetDescriptionText().ToString();
-	FilesInCL = ChangelistState->GetFilesStates();
-	ShelvedFilesInCL = ChangelistState->GetShelvedFilesStates();
+	BroadcastReadyIsValid.Emplace(Changelist.bIsValid);
 }
 
-void FChangelistService::UpdateCachedChangelistState()
+void FChangelistService::UpdateCachedChangelistData()
 {
-	check(ChangelistState);
+	TRACE_CPUPROFILER_EVENT_SCOPE(FChangelistService::DepotFileToLocalFile);
+	check(ParseChangelistTask.IsCompleted());
+	FChangelistData NewChangelist = ParseChangelistTask.GetResult();
 
 	ETaskArea ChangeType = ETaskArea::None;
 
-	if(!AreCLDescriptionsIdentical())
+	if (NewChangelist.Description.TrimEnd() != Changelist.Description.TrimEnd())
 	{
-		ChangeType |= ETaskArea::Changelist;
-
 		UE_LOG(LogSubmitToolP4, Log, TEXT("CL %s Description was updated outside of Submit Tool while it was still open, Description has been updated to match P4V."), *GetCLID());
-		UE_LOG(LogSubmitToolP4Debug, Log, TEXT("\n - Original Description '%s'\n - Submit Tool Description '%s'\n - New Description '%s'"), *OriginalDescription.ToString(), *CLDescription, *ChangelistState->GetDescriptionText().ToString());
+		UE_LOG(LogSubmitToolP4Debug, Log, TEXT("\n - Original Description '%s'\n - Submit Tool Description '%s'\n - New Description '%s'"),
+			*Changelist.Description, *CLDescription, *NewChangelist.Description);
 
-		OriginalDescription = ChangelistState->GetDescriptionText();
-		CLDescription = ChangelistState->GetDescriptionText().ToString();
+		ChangeType |= ETaskArea::Changelist;
 	}
 
-
-	bool bNewLocalFiles = false;
-	if (ChangelistState->GetFilesStatesNum() == FilesInCL.Num())
-	{
-		const TArray<FSCFileRef>& PendingFileStates = ChangelistState->GetFilesStates();
-		Algo::StableSortBy(PendingFileStates, [](const FSCFileRef& FileState) { return FileState->GetFilename(); });
-		Algo::StableSortBy(FilesInCL, [](const FSCFileRef& FileState) { return FileState->GetFilename(); });
-
-		for (size_t i = 0; i < PendingFileStates.Num(); ++i)
+	bool bNewLocalFiles = !Algo::Compare(NewChangelist.Files, Changelist.Files,
+		[](const FSCFileRef& A, const FSCFileRef& B)
 		{
-			if (!PendingFileStates[i]->GetFilename().Equals(FilesInCL[i]->GetFilename(), ESearchCase::IgnoreCase))
-			{
-				bNewLocalFiles = true;
-				break;
-			}
-		}
-	}
-	else
-	{
-		bNewLocalFiles = true;
-	}
+			return A->GetFilename().Equals(B->GetFilename(), ESearchCase::IgnoreCase) && A->GetState() == B->GetState();
+		});
 
-	if(bNewLocalFiles)
+	if (bNewLocalFiles)
 	{
 		UE_LOG(LogSubmitToolP4, Log, TEXT("CL %s files were updated outside of Submit Tool while it was open, Validation state has been reset"), *GetCLID());
-		UE_LOG(LogSubmitToolP4Debug, Log, TEXT("Current ST File info:"));
-		PrintFilesAndShelvedFiles();
 		ChangeType |= ETaskArea::LocalFiles;
-		FilesInCL = ChangelistState->GetFilesStates();
-		FilesDepotPaths.Reset(FilesInCL.Num());
-		UE_LOG(LogSubmitToolP4Debug, Log, TEXT("New P4 File info:"));
-		PrintFilesAndShelvedFiles();
 	}
 
-	bool bNewShelf = false;
-	if (ChangelistState->GetShelvedFilesStatesNum() == ShelvedFilesInCL.Num())
-	{
-		const TArray<FSCFileRef>& P4Shelf = ChangelistState->GetShelvedFilesStates();
-		Algo::StableSortBy(P4Shelf, [](const FSCFileRef& FileState) { return FileState->GetFilename(); });
-		Algo::StableSortBy(ShelvedFilesInCL, [](const FSCFileRef& FileState) { return FileState->GetFilename(); });
-
-		for (size_t i = 0; i < P4Shelf.Num(); ++i)
+	bool bNewShelf = !Algo::Compare(NewChangelist.ShelvedFiles, Changelist.ShelvedFiles,
+		[](const FSCFileRef& A, const FSCFileRef& B)
 		{
-			if (!P4Shelf[i]->GetFilename().Equals(ShelvedFilesInCL[i]->GetFilename(), ESearchCase::IgnoreCase))
-			{
-				bNewShelf = true;
-				break;
-			}
-		}
-	}
-	else
+			return A->GetDepotPath().Equals(B->GetDepotPath(), ESearchCase::IgnoreCase) && A->GetState() == B->GetState();
+		});
+
+	if (bNewShelf)
 	{
-		bNewShelf = true;
+		UE_LOG(LogSubmitToolP4, Log, TEXT("CL %s shelved files were updated outside of Submit Tool while it was open, Validation state has been reset"), *GetCLID());
+		ChangeType |= ETaskArea::ShelvedFiles;
 	}
 
-	if(bNewShelf)
+	if (bNewLocalFiles || bNewShelf)
 	{
-		UE_LOG(LogSubmitToolP4, Log, TEXT("CL %s shelved files were updated outside of Submit Tool while it was open, Validation state has been reset"), *GetCLID());
 		UE_LOG(LogSubmitToolP4Debug, Log, TEXT("Current ST File info:"));
 		PrintFilesAndShelvedFiles();
-		ChangeType |= ETaskArea::ShelvedFiles;
-		ShelvedFilesInCL = ChangelistState->GetShelvedFilesStates();
-		ShelvedFilesDepotPaths.Reset(ShelvedFilesInCL.Num());
-		UE_LOG(LogSubmitToolP4Debug, Log, TEXT("New P4 File info:"));
-		PrintFilesAndShelvedFiles();
 	}
 
-	CLRefreshCallback.ExecuteIfBound(ChangeType);
-}
-
-bool FChangelistService::AreCLDescriptionsIdentical() const
-{
-	const FString NewDescription = ChangelistState->GetDescriptionText()
-		.ToString()
-		.TrimStartAndEnd()
-		.TrimChar('\n');
+	CLDescription = NewChangelist.Description;
+	Changelist = NewChangelist;
 
-	const FString CurrentDescription = OriginalDescription
-		.ToString()
-		.TrimStartAndEnd()
-		.TrimChar('\n');
-
-	return NewDescription.Equals(CurrentDescription, ESearchCase::IgnoreCase);
-}
-
-void FChangelistService::PrintMessages(const FSourceControlResultInfo& ResultInfo) const
-{
-	for(const FText& Msg : ResultInfo.InfoMessages)
+	if (bNewLocalFiles || bNewShelf)
 	{
-		UE_LOG(LogSubmitToolP4Debug, Log, TEXT("%s"), *Msg.ToString());
+		UE_LOG(LogSubmitToolP4Debug, Log, TEXT("New P4 File info:"));
+		PrintFilesAndShelvedFiles();
 	}
-}
 
-void FChangelistService::PrintErrorMessages(const FSourceControlResultInfo& ResultInfo) const
-{
-	for(const FText& errorMsg : ResultInfo.ErrorMessages)
-	{
-		UE_LOG(LogSubmitToolP4, Error, TEXT("%s"), *errorMsg.ToString());
-	}
+	BroadcastRefreshChangeType.Emplace(ChangeType);
 }
 
-void FChangelistService::Submit(const FString& InDescriptionAddendum, const FSourceControlOperationComplete& OnSubmitComplete)
+void FChangelistService::Submit(const FString& InDescriptionAddendum, TFunction<void(const FSCCommand&)> OnComplete)
 {
+	TRACE_CPUPROFILER_EVENT_SCOPE(FChangelistData::Submit);
 	UE_LOG(LogSubmitToolP4, Warning, TEXT("Submit in progress for CL: %s. Please wait..."), *GetCLID());
 
-	TSharedRef<FCheckIn, ESPMode::ThreadSafe> CheckinOp = ISourceControlOperation::Create<FCheckIn>();
+	FString FinalDescription = WriteToString<1024>(CLDescription, InDescriptionAddendum).ToString();
+
+	FSCCommand SubmitCommand = FSCCommand(TEXT("submit"), { TEXT("-f"), TEXT("revertunchanged"),  TEXT("-i") });
 
-	FString FinalDescription = FString::Format(TEXT("{0}{1}"), { CLDescription, InDescriptionAddendum });
 
-	CheckinOp->SetDescription(FText::FromString(FinalDescription));
+	TArray<FString> DepotPathFiles;
 
-	ActiveP4Operations.Add(CheckinOp);
+	Algo::Transform(Changelist.Files, DepotPathFiles, [](const FSCFileRef& InFile) { return InFile->GetDepotPath(); });
 
-	CachedSCCProvider->Execute(CheckinOp, ChangelistPtr, EConcurrency::Asynchronous, FSourceControlOperationComplete::CreateLambda(
-		[this, OnSubmitComplete](const FSourceControlOperationRef& Operation, ECommandResult::Type Result) {
-			if(Result == ECommandResult::Failed)
+	// -i needs each description line and each depot filepath to be prepended by \t
+	FinalDescription.ReplaceInline(TEXT("\n"), TEXT("\n\t"));
+	SubmitCommand.InputData = FString::Printf(TEXT("Change: %s\nDescription:\n\t%s\nFiles:\n\t%s"), *CLID, *FinalDescription, *FSubmitToolCoreUtils::StringBuilderJoin<2048>(DepotPathFiles, TEXT("\n\t")));
+
+	SourceControlService->RunCommand(MoveTemp(SubmitCommand), FOnSCCCommandComplete::CreateLambda(
+		[OnComplete, this](const FSCCommand& Cmd)
+		{
+			TRACE_CPUPROFILER_EVENT_SCOPE(FChangelistData::Submit_callback);
+
+			if (Cmd.bSuccess)
 			{
-				UE_LOG(LogSubmitToolP4, Error, TEXT("CL %s submit operation failed:"), *GetCLID());
-				PrintErrorMessages(Operation->GetResultInfo());
+				UE_LOG(LogSubmitToolP4, Log, TEXT("Submit operation succeeded"));
 			}
-			else if(Result == ECommandResult::Succeeded)
+			else if(Cmd.bDroppedConnection)
+			{ 
+				UE_LOG(LogSubmitToolP4, Warning, TEXT("Submit operation cancelled."));
+			}
+			else
 			{
-				TSharedRef<FCheckIn> CheckIn = StaticCastSharedRef<FCheckIn>(Operation);
-				UE_LOG(LogSubmitToolP4, Log, TEXT("Submit operation succeeded: %s"), *CheckIn->GetSuccessMessage().ToString());
+				UE_LOG(LogSubmitToolP4, Warning, TEXT("Submit operation failed."));
 			}
-			else if(Result == ECommandResult::Cancelled)
+
+			if (OnComplete.IsSet())
 			{
-				UE_LOG(LogSubmitToolP4, Warning, TEXT("Submit operation cancelled."));
+				OnComplete(Cmd);
 			}
 
-			ActiveP4Operations.Remove(Operation);
-			OnSubmitComplete.ExecuteIfBound(Operation, Result);
-		}
-	));
+		}));
 }
 
-void FChangelistService::SendCLDescriptionToP4(EConcurrency::Type Concurrency, FOnP4OperationCompleteDelegate InCallback)
+void FChangelistService::DeleteShelvedFiles(TFunction<void(const FSCCommand&)> OnDeleteComplete)
 {
-	if(!CLDescription.Equals(OriginalDescription.ToString()))
-	{
-		UE_LOG(LogSubmitToolP4Debug, Log, TEXT("Saving CL %s"), *GetCLID());
-		TSharedRef<FEditChangelist, ESPMode::ThreadSafe> EditCLOp = ISourceControlOperation::Create<FEditChangelist>();
-		EditCLOp->SetDescription(FText::FromString(CLDescription));
-
-		ActiveP4Operations.Add(EditCLOp);
+	TRACE_CPUPROFILER_EVENT_SCOPE(FChangelistService::DeleteShelvedFiles);
+	UE_LOG(LogSubmitToolP4, Log, TEXT("Removing shelved files in CL %s..."), *GetCLID());
 
-		CachedSCCProvider->Execute(EditCLOp, ChangelistPtr, Concurrency, FSourceControlOperationComplete::CreateLambda(
-			[this, InCallback](const FSourceControlOperationRef& UpdateOperation, ECommandResult::Type Result) {
-				if(Result == ECommandResult::Failed)
-				{
-					UE_LOG(LogSubmitToolP4, Error, TEXT("CL %s edit changelist operation %s failed:"), *GetCLID(), *UpdateOperation->GetName().ToString());
-					PrintErrorMessages(UpdateOperation->GetResultInfo());
-				}
-				else if(Result == ECommandResult::Succeeded)
-				{
-					UE_LOG(LogSubmitToolP4, Log, TEXT("CL %s edit changelist operation succeeded."), *GetCLID());
-					OriginalDescription = FText::FromString(CLDescription);
-				}
-				else if(Result == ECommandResult::Cancelled)
-				{
-					UE_LOG(LogSubmitToolP4, Warning, TEXT("CL %s edit changelist operation cancelled."), *GetCLID());
-				}
+	SourceControlService->RunCommand(TEXT("shelve"), { TEXT("-d") , TEXT("-c"), GetCLID() }, FOnSCCCommandComplete::CreateLambda(
+		[OnDeleteComplete, this](const FSCCommand& Cmd)
+		{
+			TRACE_CPUPROFILER_EVENT_SCOPE(FChangelistService::DeleteShelvedFiles_callback);
 
-				ActiveP4Operations.Remove(UpdateOperation);
+			if (Cmd.bSuccess)
+			{
+				Changelist.ShelvedFiles.Reset();
+				UE_LOG(LogSubmitToolP4, Log, TEXT("CL %s delete shelf operation succeeded"), *GetCLID());
+			}
+			else if (Cmd.bDroppedConnection)
+			{
+				UE_LOG(LogSubmitToolP4, Log, TEXT("CL %s delete shelf operation cancelled"), *GetCLID());
+			}
+			else
+			{
+				UE_LOG(LogSubmitToolP4, Error, TEXT("CL %s delete shelf operation failed:"), *GetCLID());
+			}
 
-				InCallback.ExecuteIfBound();
+			if (OnDeleteComplete.IsSet())
+			{
+				OnDeleteComplete(Cmd);
 			}
-		));
-	}
-	else
-	{
-		InCallback.ExecuteIfBound();
-	}
+		}
+	));
 }
 
-void FChangelistService::DeleteShelvedFiles(const FSourceControlOperationComplete& OnDeleteComplete)
+void FChangelistService::ReplaceShelvedFiles(TFunction<void(const FSCCommand&)> OnComplete)
 {
-	UE_LOG(LogSubmitToolP4, Log, TEXT("Removing shelved files in CL %s..."), *GetCLID());
-	const TSharedRef<FDeleteShelved, ESPMode::ThreadSafe> DeleteShelvedOp = ISourceControlOperation::Create<FDeleteShelved>();
+	TRACE_CPUPROFILER_EVENT_SCOPE(FChangelistService::ReplaceShelvedFiles);
+	UE_LOG(LogSubmitToolP4, Log, TEXT("Replacing shelved files in CL %s..."), *GetCLID());
 
-	ActiveP4Operations.Add(DeleteShelvedOp);
+	SourceControlService->RunCommand(TEXT("shelve"), { TEXT("-r") , TEXT("-c"), GetCLID() }, FOnSCCCommandComplete::CreateLambda(
+		[OnComplete, this](const FSCCommand& Cmd)
+		{
+			TRACE_CPUPROFILER_EVENT_SCOPE(FChangelistService::ReplaceShelvedFiles_callback);
 
-	CachedSCCProvider->Execute(DeleteShelvedOp, ChangelistPtr, EConcurrency::Asynchronous, FSourceControlOperationComplete::CreateLambda(
-		[this, OnDeleteComplete](const FSourceControlOperationRef& DeleteShelvedOp, ECommandResult::Type Result) {
-			if(Result == ECommandResult::Failed)
+			if (Cmd.bSuccess)
 			{
-				UE_LOG(LogSubmitToolP4, Error, TEXT("CL %s delete shelved operation failed:"), *GetCLID());
+				UE_LOG(LogSubmitToolP4, Log, TEXT("CL %s replace shelf operation succeeded"), *GetCLID());
+				Changelist.ShelvedFiles = ParseShelvedFilesFromP4(Cmd);
 
-				PrintErrorMessages(DeleteShelvedOp->GetResultInfo());
+				TSharedPtr<ICacheDataService> CacheService = ServiceProvider->GetService<ICacheDataService>();
+
+				if (CacheService.IsValid())
+				{
+					CacheService->SetLastShelfKnownUpToDate(GetCLID(), FDateTime::UtcNow());
+				}
+				bIsShelfReady = true;
 			}
-			else if(Result == ECommandResult::Succeeded)
+			else if (Cmd.bDroppedConnection)
 			{
-				ShelvedFilesInCL.Reset();
-				ShelvedFilesDepotPaths.Reset();
-				UE_LOG(LogSubmitToolP4, Log, TEXT("CL %s delete shelved operation succeeded"), *GetCLID());
+				UE_LOG(LogSubmitToolP4, Log, TEXT("CL %s replace shelf operation cancelled"), *GetCLID());
 			}
-			else if(Result == ECommandResult::Cancelled)
+			else
 			{
-				UE_LOG(LogSubmitToolP4, Warning, TEXT("CL %s delete shelved operation cancelled."), *GetCLID());
+				UE_LOG(LogSubmitToolP4, Error, TEXT("CL %s replace shelf operation failed"), *GetCLID());
 			}
 
-			ActiveP4Operations.Remove(DeleteShelvedOp);
-			OnDeleteComplete.ExecuteIfBound(DeleteShelvedOp, Result);
+			if (OnComplete.IsSet())
+			{
+				OnComplete(Cmd);
+			}
 		}
 	));
 }
 
-void FChangelistService::CreateShelvedFiles(const FSourceControlOperationComplete& OnCreateComplete)
+void FChangelistService::CreateShelvedFiles(TFunction<void(const FSCCommand&)> OnComplete)
 {
-	UE_LOG(LogSubmitToolP4, Log, TEXT("Creating shelved files for CL %s..."), *GetCLID());
-	const TSharedRef<FShelve, ESPMode::ThreadSafe> CreateShelvedOp = ISourceControlOperation::Create<FShelve>();
-
-	ActiveP4Operations.Add(CreateShelvedOp);
+	TRACE_CPUPROFILER_EVENT_SCOPE(FChangelistService::CreateShelvedFiles);
+	UE_LOG(LogSubmitToolP4, Log, TEXT("Creating shelf for CL %s..."), *GetCLID());
 
-	CachedSCCProvider->Execute(CreateShelvedOp, ChangelistPtr, EConcurrency::Asynchronous, FSourceControlOperationComplete::CreateLambda(
-		[this, OnCreateComplete](const FSourceControlOperationRef& CreateShelvedOp, ECommandResult::Type Result) {
-			if(Result == ECommandResult::Failed)
-			{
-				UE_LOG(LogSubmitToolP4, Error, TEXT("CL %s create shelved operation failed:"), *GetCLID());
+	SourceControlService->RunCommand(TEXT("shelve"), { TEXT("-c"), GetCLID() }, FOnSCCCommandComplete::CreateLambda(
+		[OnComplete, this](const FSCCommand& Cmd)
+		{
+			TRACE_CPUPROFILER_EVENT_SCOPE(FChangelistService::CreateShelvedFiles_callback);
 
-				PrintErrorMessages(CreateShelvedOp->GetResultInfo());
-			}
-			else if(Result == ECommandResult::Succeeded)
+			if (Cmd.bSuccess)
 			{
-				ShelvedFilesInCL = ChangelistState->GetShelvedFilesStates();
-				ShelvedFilesDepotPaths.Reset(ShelvedFilesInCL.Num());
+				UE_LOG(LogSubmitToolP4, Log, TEXT("CL %s create shelf operation succeeded"), *GetCLID());
+				Changelist.ShelvedFiles = ParseShelvedFilesFromP4(Cmd);
 
 				TSharedPtr<ICacheDataService> CacheService = ServiceProvider->GetService<ICacheDataService>();
-
 				if (CacheService.IsValid())
 				{
 					CacheService->SetLastShelfKnownUpToDate(GetCLID(), FDateTime::UtcNow());
 				}
+
 				bIsShelfReady = true;
-				UE_LOG(LogSubmitToolP4, Log, TEXT("CL %s create shelved operation succeeded"), *GetCLID());
 			}
-			else if(Result == ECommandResult::Cancelled)
+			else if (Cmd.bDroppedConnection)
+			{
+				UE_LOG(LogSubmitToolP4, Log, TEXT("CL %s create shelf operation cancelled"), *GetCLID());
+			}
+			else
 			{
-				UE_LOG(LogSubmitToolP4, Warning, TEXT("CL %s create shelved operation cancelled."), *GetCLID());
+				UE_LOG(LogSubmitToolP4, Error, TEXT("CL %s create shelf operation failed"), *GetCLID());
 			}
 
-			ActiveP4Operations.Remove(CreateShelvedOp);
-			OnCreateComplete.ExecuteIfBound(CreateShelvedOp, Result);
+			if (OnComplete.IsSet())
+			{
+				OnComplete(Cmd);
+			}
 		}
 	));
 }
 
-void FChangelistService::UpdateShelvedFiles(TArray<FString> InFiles, TFunction<void(bool)> OnComplete)
+void FChangelistService::UpdateShelvedFiles(TArray<FString> InFiles, TFunction<void(const FSCCommand&)> OnComplete)
 {
 	TRACE_CPUPROFILER_EVENT_SCOPE(FChangelistService::UpdateShelvedFiles);
 	InFiles.Insert(TEXT("-f"), 0);
@@ -501,6 +560,8 @@ void FChangelistService::UpdateShelvedFiles(TArray<FString> InFiles, TFunction<v
 	SourceControlService->RunCommand(TEXT("shelve"), MoveTemp(InFiles), FOnSCCCommandComplete::CreateLambda(
 		[OnComplete, this](const FSCCommand& Cmd)
 		{
+			TRACE_CPUPROFILER_EVENT_SCOPE(FChangelistService::UpdateShelvedFiles_callback);
+
 			if (Cmd.bSuccess)
 			{
 				TSharedPtr<ICacheDataService> CacheService = ServiceProvider->GetService<ICacheDataService>();
@@ -512,56 +573,41 @@ void FChangelistService::UpdateShelvedFiles(TArray<FString> InFiles, TFunction<v
 				UE_LOG(LogSubmitToolP4, Log, TEXT("CL %s shelved files have been updated"), *GetCLID());
 				bIsShelfReady = true;
 			}
+			else if (Cmd.bDroppedConnection)
+			{
+				UE_LOG(LogSubmitToolP4, Log, TEXT("Update shelf operation was cancelled"), *GetCLID());
+			}
+			else
+			{
+				UE_LOG(LogSubmitToolP4, Error, TEXT("Update shelf operation failed"), *GetCLID());
+			}
 
 			if (OnComplete != nullptr)
 			{
-				OnComplete(Cmd.bSuccess);
+				OnComplete(Cmd);
 			}
 		}));
 }
 
 bool FChangelistService::P4Tick(float DeltaTime)
 {
-	for(const TSharedRef<ISourceControlOperation>& operation : ActiveP4Operations)
+	if (BroadcastReadyIsValid.IsSet())
 	{
-		UE_LOG(LogSubmitToolP4, Log, TEXT("%s operation still in progress: %s"), *operation->GetName().ToString(), *operation->GetInProgressString().ToString());
+		CLReadyCallback.ExecuteIfBound(BroadcastReadyIsValid.GetValue());
+		BroadcastReadyIsValid.Reset();
 	}
-
-	return true;
-}
-
-
-bool FChangelistService::IsP4OperationRunning(FName OperationName)
-{
-	if(OperationName == FName())
+	if (BroadcastRefreshChangeType.IsSet())
 	{
-		return ActiveP4Operations.Num() > 0 || SourceControlService->IsBusy();
+		CLRefreshCallback.ExecuteIfBound(BroadcastRefreshChangeType.GetValue());
+		BroadcastRefreshChangeType.Reset();
 	}
-
-	for(const TSharedRef<ISourceControlOperation>& Operation : ActiveP4Operations)
-	{
-		if(Operation->GetName().IsEqual(OperationName, ENameCase::IgnoreCase))
-		{
-			return true;
-		}
-	}
-
-	return false;
+	return true;
 }
 
-void FChangelistService::CancelP4Operations(FName OperationName)
+void FChangelistService::CancelP4Operations()
 {
-	for(const TSharedRef<ISourceControlOperation>& Operation : ActiveP4Operations)
-	{
-		if(OperationName.IsNone() || Operation->GetName().IsEqual(OperationName, ENameCase::IgnoreCase))
-		{
-			if(CachedSCCProvider->CanCancelOperation(Operation))
-			{
-				CachedSCCProvider->CancelOperation(Operation);
-				UE_LOG(LogSubmitToolP4, Warning, TEXT("P4 Operation %s cancelling requested"), *Operation->GetName().ToString())
-			}
-		}
-	}
+	UE_LOG(LogSubmitToolP4, Warning, TEXT("P4 Operations cancel requested"))
+	SourceControlService->RequestCancel();
 }
 
 void FChangelistService::CreateCLFromDefaultCL()
@@ -570,248 +616,214 @@ void FChangelistService::CreateCLFromDefaultCL()
 
 	// If we try to open submit tool with the default cl: Update all CL Status -> create a new CL -> move files from default to the new CL -> regular flow
 	UE_LOG(LogSubmitToolP4, Log, TEXT("Default changelist is not supported by Submit Tool, creating a new CL and moving files..."));
+	SourceControlService->RunCommand(TEXT("change"), { TEXT("-o") }, FOnSCCCommandComplete::CreateLambda(
+		[this](const FSCCommand& InCmd) {
+			TRACE_CPUPROFILER_EVENT_SCOPE(FChangelistService::CreateCLFromDefaultCL_Change_Output);
 
-	TSharedRef<FUpdatePendingChangelistsStatus> UpdatePendingChangelistsOperation = ISourceControlOperation::Create<FUpdatePendingChangelistsStatus>();
-	UpdatePendingChangelistsOperation->SetUpdateFilesStates(true);
-	UpdatePendingChangelistsOperation->SetUpdateAllChangelists(true);
-
-	ActiveP4Operations.Add(UpdatePendingChangelistsOperation);
-
-	CachedSCCProvider->Execute(UpdatePendingChangelistsOperation, EConcurrency::Asynchronous, FSourceControlOperationComplete::CreateLambda(
-		[this](const FSourceControlOperationRef& Operation, ECommandResult::Type Result) {
-			TRACE_CPUPROFILER_EVENT_SCOPE(FChangelistService::CreateCLFromDefaultCL_Callback);
-
-			ActiveP4Operations.Remove(Operation);
-
-			if(Result == ECommandResult::Failed)
-			{
-				UE_LOG(LogSubmitToolP4, Error, TEXT("Failed to fetch changelists from p4."));
-				PrintErrorMessages(Operation->GetResultInfo());
-			}
-			else if(Result == ECommandResult::Succeeded)
+			if (InCmd.bSuccess)
 			{
+				FSCCommand Command = FSCCommand(TEXT("change"), { TEXT("-i") });
+				Command.InputData = ParseP4ChangesOutput(InCmd.Values[0], Parameters.NewChangelistMessage);
 
-				TArray<FString> FilesInDefault;
-				for(const FSourceControlChangelistPtr cl : CachedSCCProvider->GetChangelists(EStateCacheUsage::Use))
-				{
-					if(cl->IsDefault())
-					{
-						for(FSCFileRef file : CachedSCCProvider->GetState(cl.ToSharedRef(), EStateCacheUsage::Use)->GetFilesStates())
+				SourceControlService->RunCommand(MoveTemp(Command), FOnSCCCommandComplete::CreateLambda(
+					[this](const FSCCommand& InCmd) {
+						TRACE_CPUPROFILER_EVENT_SCOPE(FChangelistService::CreateCLFromDefaultCL_Change_Input);
+						if (InCmd.bSuccess)
 						{
-							FilesInDefault.Emplace(file->GetFilename());
-						}
-					}
-				}
-
-				TSharedRef<FNewChangelist, ESPMode::ThreadSafe> NewCLOp = ISourceControlOperation::Create<FNewChangelist>();
-
-				FStringFormatNamedArguments FormatArgs = { 
-					{ TEXT("FileCount") ,FilesInDefault.Num() }
-				};
-
-				NewCLOp->SetDescription(FText::FromString(FString::Format(*Parameters.NewChangelistMessage, FormatArgs)));
-
-				ActiveP4Operations.Add(NewCLOp);
-
-				CachedSCCProvider->Execute(NewCLOp, EConcurrency::Asynchronous, FSourceControlOperationComplete::CreateLambda(
-					[this, FilesInDefault](const FSourceControlOperationRef& Operation, ECommandResult::Type Result) {
-						ActiveP4Operations.Remove(Operation);
-
-						if(Result == ECommandResult::Failed)
-						{
-							UE_LOG(LogSubmitToolP4, Error, TEXT("Failed to create new changelist from default"));
-							PrintErrorMessages(Operation->GetResultInfo());
-						}
-						else if(Result == ECommandResult::Succeeded)
-						{
-							TSharedRef<FNewChangelist> NewCLOp = StaticCastSharedRef<FNewChangelist>(Operation);
-							ChangelistPtr = NewCLOp->GetNewChangelist();
-							ChangelistState = CachedSCCProvider->GetState(ChangelistPtr.ToSharedRef(), EStateCacheUsage::Use);
-
-							OriginalDescription = ChangelistState->GetDescriptionText();
-							CLDescription = ChangelistState->GetDescriptionText().ToString();
-
-
-							TSharedRef<FMoveToChangelist> MoveOp = ISourceControlOperation::Create<FMoveToChangelist>();
-
-							ActiveP4Operations.Add(MoveOp);
-							CachedSCCProvider->Execute(MoveOp, ChangelistPtr, FilesInDefault, EConcurrency::Asynchronous, FSourceControlOperationComplete::CreateLambda(
-								[this, FileCount = FilesInDefault.Num()](const FSourceControlOperationRef& Operation, ECommandResult::Type Result) {
-									ActiveP4Operations.Remove(Operation);
-									if(Result == ECommandResult::Failed)
-									{
-										UE_LOG(LogSubmitToolP4, Error, TEXT("Move files to CL failed"));
-										PrintErrorMessages(Operation->GetResultInfo());
-									}
-									else if(Result == ECommandResult::Succeeded)
-									{
-										FilesInCL = ChangelistState->GetFilesStates();
-										FilesDepotPaths.Reset(FilesInCL.Num());
-										ShelvedFilesInCL = ChangelistState->GetShelvedFilesStates();
-										ShelvedFilesDepotPaths.Reset(ShelvedFilesInCL.Num());
-										CLID = ChangelistPtr->GetIdentifier();
-										FConfiguration::AddOrUpdateEntry(TEXT("$(CL)"), GetCLID());
-
-										UE_LOG(LogSubmitToolP4, Log, TEXT("Created CL %s and moved with %d files from the default CL."), *ChangelistPtr->GetIdentifier(), FileCount);
-										UE_LOG(LogSubmitToolDebug, Log, TEXT("ChangeListService CLID Updated to %s"), *GetCLID());
-										UE_LOG(LogSubmitToolDebug, Log, TEXT("Configuration updated value $(CL) to %s"), *GetCLID());
-
-										CLReadyCallback.ExecuteIfBound(ChangelistState.IsValid());
-									}
-									else if(Result == ECommandResult::Cancelled)
+							for (const FText& Message : InCmd.Results.InfoMessages)
+							{
+								// Perforce message format is "Change XXXXXX created ..." it doesn't return it as a value so parse it from the message
+								const FString StringMessage = Message.ToString();
+								TStringBuilder<16> CLBuilder;
+								for (int32 i = 7; i < StringMessage.Len() && StringMessage[i] != TCHAR(' '); ++i)
+								{
+									if (FChar::IsDigit(StringMessage[i]))
 									{
-										UE_LOG(LogSubmitToolP4, Warning, TEXT("Move files to CL cancelled."), *GetCLID());
+										CLBuilder.AppendChar(StringMessage[i]);
 									}
-								}));
-						}
-						else if(Result == ECommandResult::Cancelled)
-						{
-							UE_LOG(LogSubmitToolP4, Warning, TEXT("Create new CL was cancelled."), *GetCLID());
+								}
+
+								if (CLBuilder.Len() != 0)
+								{
+									CLID = CLBuilder.ToString();
+									FConfiguration::AddOrUpdateEntry(TEXT("$(CL)"), GetCLID());
+									UE_LOG(LogSubmitToolP4, Log, TEXT("Created CL %s from default."), *CLID);
+								}
+								else
+								{
+									UE_LOG(LogSubmitToolP4Debug, Error, TEXT("Couldn't retrieve the created CL number."));
+								}
+							}
+
+							FetchChangelistDataAsync();
 						}
 					}
 				));
 			}
-			else if(Result == ECommandResult::Cancelled)
-			{
-				UE_LOG(LogSubmitToolP4, Warning, TEXT("CL %s edit changelist operation cancelled."), *GetCLID());
-			}
-		}));
+		}
+	));
 }
 
 void FChangelistService::PrintFilesAndShelvedFiles() const
 {
-	if(FilesInCL.Num() > 0)
+	if (Changelist.Files.Num() > 0)
 	{
 		UE_LOG(LogSubmitToolP4Debug, Log, TEXT("Files in CL:"));
-		for(const FSCFileRef& File : FilesInCL)
+		for(const FSCFileRef& File : Changelist.Files)
 		{
-			UE_LOG(LogSubmitToolP4Debug, Log, TEXT("\t%s"), *File->GetFilename());
+			UE_LOG(LogSubmitToolP4Debug, Log, TEXT("\t%s - %s"), *File->GetFilename(), *File->GetDepotPath());
 		}
 	}
 
-	if(ShelvedFilesInCL.Num() > 0)
+	if (Changelist.ShelvedFiles.Num() > 0)
 	{
 		UE_LOG(LogSubmitToolP4Debug, Log, TEXT("Shelved Files in CL:"));
-		for(const FSCFileRef& File : ShelvedFilesInCL)
+		for(const FSCFileRef& File : Changelist.ShelvedFiles)
 		{
-			UE_LOG(LogSubmitToolP4Debug, Log, TEXT("\t%s"), *File->GetFilename());
+			UE_LOG(LogSubmitToolP4Debug, Log, TEXT("\t%s - %s"), *File->GetFilename(), *File->GetDepotPath());
 		}
 	}
 }
-UE::Tasks::TTask<TSharedRef<FSCCommand>> FChangelistService::LaunchGetOpenedFilesDepotPaths() const
-{
-	return SourceControlService->RunCommand(TEXT("opened"), { });
-}
 
-void FChangelistService::ParseOpenedFilesDepotPaths(TSharedRef<FSCCommand> InCompletedSCCommand)
+TArray<FSCFileRef> FChangelistService::ParseShelvedFilesFromP4(const FSCCommand& InCmd)
 {
-	FilesDepotPaths.Reset();
-	OpenFilesInChangelists.Reset();
-	if (InCompletedSCCommand->bSuccess)
+	TArray<FSCFileRef> NewShelvedFiles;
+	TArray<FString> WhereFiles;
+	TMap<FString, ESCFileState> WhereStates;
+	for (const FSCCRecord& Record : InCmd.Values)
 	{
-		for (const TMap<FString, FString>& Record : InCompletedSCCommand->Values)
+		const FString* DepotFile = Record.Find(TEXT("depotFile"));
+		const FString* Action = Record.Find(TEXT("action"));
+		if (DepotFile && Action)
 		{
-			const FString* DepotFile = Record.Find(TEXT("depotFile"));
-			const FString* Change = Record.Find(TEXT("change"));
-			if (DepotFile && Change)
+			ESCFileState State = GetLocalFileState(*Action);
+
+			// Check if it was already in the shelf before
+			FSCFileRef* FoundFile = Changelist.ShelvedFiles.FindByPredicate([&DepotFile](const FSCFileRef& InFileRef) { return InFileRef->GetDepotPath().Equals(*DepotFile, ESearchCase::IgnoreCase); });
+			if (FoundFile != nullptr)
 			{
-				if (*Change == CLID)
+				(*FoundFile)->UpdateState(State);
+				NewShelvedFiles.Add(*FoundFile);
+			}
+			else
+			{
+				// Look into the local files to see if we know about locations for this one already
+				FoundFile = Changelist.Files.FindByPredicate([&DepotFile](const FSCFileRef& InFileRef) { return InFileRef->GetDepotPath().Equals(*DepotFile, ESearchCase::IgnoreCase); });
+				if (FoundFile != nullptr)
+				{
+					NewShelvedFiles.Add(MakeShared<FSCFile>((*FoundFile)->GetFilename(), *DepotFile, State));
+				}
+				else
 				{
-					FilesDepotPaths.Add(*DepotFile);
+					// Needs a p4 where to figure out the local filepath
+					WhereFiles.Add(*DepotFile);
+					WhereStates.Add(*DepotFile, State);
 				}
-				OpenFilesInChangelists.Add(*DepotFile, *Change);
 			}
 		}
 	}
-}
-UE::Tasks::TTask<TSharedRef<FSCCommand>> FChangelistService::LaunchGetShelvedFilesDepotPaths() const
-{
-	return SourceControlService->RunCommand(TEXT("describe"), { "-S", CLID });
-}
 
-void FChangelistService::ParseShelvedFilesDepotPaths(TSharedRef<FSCCommand> InCompletedSCCommand)
-{
-	ShelvedFilesDepotPaths.Reset();
-	if (InCompletedSCCommand->bSuccess)
+	if (WhereFiles.Num() > 0)
 	{
-		for (const TMap<FString, FString>& Record : InCompletedSCCommand->Values)
+		TMap<FString, FString> WhereResults = RunWhereCommandSync(MoveTemp(WhereFiles));
+
+		for (const TPair<FString, FString>& WhereEntry : WhereResults)
 		{
-			int32 I = 0;
-			FString PathsKey = FString::Printf(TEXT("depotFile%d"), I);
-			while (Record.Contains(PathsKey))
-			{
-				ShelvedFilesDepotPaths.Add(Record[PathsKey]);
-				++I;
-				PathsKey = FString::Printf(TEXT("depotFile%d"), I);
-			}
+			NewShelvedFiles.Add(MakeShared<FSCFile>(WhereEntry.Value, WhereEntry.Key, WhereStates[WhereEntry.Key]));
 		}
 	}
+
+	return NewShelvedFiles;
 }
 
-void FChangelistService::WaitForInitialP4RunTasks()
+TMap<FString, FString> FChangelistService::RunWhereCommandSync(TArray<FString>&& InDepotPaths)
 {
-	TRACE_CPUPROFILER_EVENT_SCOPE(FChangelistService::WaitForInitialP4RunTasks);
-
-	if (InitialP4OpenedFilesTask.IsValid())
+	TMap<FString, FString> DepotFileToLocalFile;
+	FSCCommand WhereCmd = SourceControlService->RunSyncCommand(TEXT("where"), MoveTemp(InDepotPaths));
+	if (WhereCmd.bSuccess)
 	{
-		ParseOpenedFilesDepotPaths(InitialP4OpenedFilesTask.GetResult());
-		InitialP4OpenedFilesTask = {};
+		for (const TMap<FString, FString>& Record : WhereCmd.Values)
+		{
+			const FString* DepotFile = Record.Find(TEXT("depotFile"));
+			const FString* Path = Record.Find(TEXT("path"));
+			if (DepotFile && Path)
+			{
+				DepotFileToLocalFile.Add(*DepotFile, *Path);
+			}
+		}
 	}
 
-	if (InitialP4ShelvedFilesTask.IsValid())
-	{
-		ParseShelvedFilesDepotPaths(InitialP4ShelvedFilesTask.GetResult());
-		InitialP4ShelvedFilesTask = {};
-	}
+	return DepotFileToLocalFile;
 }
-
-const TArray<FString>& FChangelistService::GetOpenedFilesDepotPaths()
+void FChangelistService::WaitForChangelistData()
 {
-	check(IsInGameThread());
-	WaitForInitialP4RunTasks();
-
-	if (FilesDepotPaths.IsEmpty() && !FilesInCL.IsEmpty())
-	{
-		TRACE_CPUPROFILER_EVENT_SCOPE(FChangelistService::GetOpenedFilesDepotPaths);
-		ParseOpenedFilesDepotPaths(LaunchGetOpenedFilesDepotPaths().GetResult());
-	}
-
-	return FilesDepotPaths;
+	TRACE_CPUPROFILER_EVENT_SCOPE(FChangelistService::WaitForChangelistData);
+	UpdateChangelistTask.Wait();
+	OpenedTask = {};
+	DescribeTask = {};
+	ParseChangelistTask = {};
+	UpdateChangelistTask = {};
 }
 
-const TArray<FString>& FChangelistService::GetShelvedFilesDepotPaths()
+FString FChangelistService::ParseP4ChangesOutput(const FSCCRecord& InP4ChangeOutputRecords, const FString& InReplaceDescription)
 {
-	check(IsInGameThread());
-	WaitForInitialP4RunTasks();
+	TRACE_CPUPROFILER_EVENT_SCOPE(FChangelistService::ParseP4ChangesOutput);
 
-	// Don't bother with the default CL
-	if (CLID.Equals("default", ESearchCase::IgnoreCase))
+	TStringBuilder<1024> Buffer;
+	for (const TPair<FString, FString>& RecordPair : InP4ChangeOutputRecords)
 	{
-		return ShelvedFilesDepotPaths;
-	}
-
-	if (ShelvedFilesDepotPaths.IsEmpty() && !ShelvedFilesInCL.IsEmpty())
-	{
-		TRACE_CPUPROFILER_EVENT_SCOPE(FChangelistService::GetShelvedFilesDepotPaths);
-		ShelvedFilesDepotPaths.Reset();
-		ParseShelvedFilesDepotPaths(LaunchGetShelvedFilesDepotPaths().GetResult());
-	}
-
-	return ShelvedFilesDepotPaths;
-}
+		// Skip some fields that aren't required
+		if (RecordPair.Key == TEXT("Date") ||
+			RecordPair.Key == TEXT("Type") ||
+			RecordPair.Key == TEXT("specFormatted") ||
+			RecordPair.Key == TEXT("func"))
+		{
+			continue;
+		}
+		else if (RecordPair.Key == TEXT("extraTag0"))
+		{
+			// Changelists with shelved files will have additional tags which will be rejected in the edit
+			break;
+		}
 
-const TMap<FString, FString>& FChangelistService::GetOpenFilesInChangelists()
-{
-	check(IsInGameThread());
-	WaitForInitialP4RunTasks();
+		if (RecordPair.Key == TEXT("Description"))
+		{
+			// Description case: we will replace the current description by the new one provided
+			Buffer.Append(TEXT("Description:\n"));
+			{
+				TArray<FString> DescLines;
+				InReplaceDescription.ParseIntoArray(DescLines, TEXT("\n"), false);
+				for (const FString& DescLine : DescLines)
+				{
+					Buffer.Append(TEXT("\t"));
+					Buffer.Append(DescLine);
+					Buffer.Append(TEXT("\n"));
+				}
+			}
+			Buffer.Append(TEXT("\n"));
+		}
+		else if (RecordPair.Key.Contains(TEXT("Files")))
+		{
+			if (RecordPair.Key == TEXT("Files0"))
+			{
+				Buffer.Append(TEXT("Files:\n"));
+			}
 
-	if (OpenFilesInChangelists.IsEmpty())
-	{
-		TRACE_CPUPROFILER_EVENT_SCOPE(FChangelistService::GetOpenFilesInChangelists);
-		ParseOpenedFilesDepotPaths(LaunchGetOpenedFilesDepotPaths().GetResult());
+			Buffer.Append(TEXT("\t"));
+			Buffer.Append(RecordPair.Value);
+			Buffer.Append(TEXT("\n\n"));
+		}
+		else
+		{
+			// General case: just put back what is already present in the record
+			Buffer.Append(RecordPair.Key);
+			Buffer.Append(TEXT(":\t"));
+			Buffer.Append(RecordPair.Value);
+			Buffer.Append(TEXT("\n\n"));
+		}
 	}
 
-	return OpenFilesInChangelists;
+	return Buffer.ToString();
 }
 
 void FChangelistService::CheckShelfIsUpToDate(TFunction<void(bool, TArray<FString>)> OnComplete)
@@ -820,7 +832,7 @@ void FChangelistService::CheckShelfIsUpToDate(TFunction<void(bool, TArray<FStrin
 
 	TArray<FString> FilePathsOutOfDate;
 
-	if (GetOpenedFilesDepotPaths() != GetShelvedFilesDepotPaths())
+	if (!Changelist.bIsShelfComplete)
 	{
 		if(OnComplete)
 		{
@@ -837,7 +849,7 @@ void FChangelistService::CheckShelfIsUpToDate(TFunction<void(bool, TArray<FStrin
 		if (LastShelf != FDateTime::MinValue())
 		{
 			FFileManagerGeneric FileManager;
-			for (const FSCFileRef& File : FilesInCL)
+			for (const FSCFileRef& File : Changelist.Files)
 			{
 				FString Filename = File->GetFilename();
 				FFileStatData FileModifiedDate = FileManager.GetStatData(*Filename);
@@ -866,11 +878,15 @@ void FChangelistService::CheckShelfIsUpToDate(TFunction<void(bool, TArray<FStrin
 	{
 		bDiffTaskRunning = true;
 		// Real p4 diff
-		UE_LOG(LogSubmitToolP4, Log, TEXT("Checking Shelf diff with local files..."));
-		TArray<FString> DiffArgs = GetOpenedFilesDepotPaths();
-		for (FString& FilePath : DiffArgs)
+		UE_LOG(LogSubmitToolP4Debug, Log, TEXT("Checking Shelf diff with local files..."));
+		TArray<FString> DiffArgs;
+		DiffArgs.Reserve(GetFilesInCL().Num());
+		for (const FSCFileRef& File : GetFilesInCL())
 		{
-			FilePath.Appendf(TEXT("@=%s"), *CLID);
+			if (!File->IsDeleted())
+			{
+				DiffArgs.Emplace(*WriteToString<256>(File->GetDepotPath(), "@=", CLID));
+			}
 		}
 		DiffArgs.Insert(TEXT("-ds"), 0);
 		DiffArgs.Insert(TEXT("-Od"), 0);
@@ -915,7 +931,7 @@ void FChangelistService::CheckShelfIsUpToDate(TFunction<void(bool, TArray<FStrin
 					OnComplete(bIsShelfReady, FilesWithChanges);
 				}
 				bDiffTaskRunning = false;
-			}));
+			}), true);
 	}
 }
 
@@ -924,7 +940,7 @@ void FChangelistService::EnsureShelfIsCurrent()
 	TRACE_CPUPROFILER_EVENT_SCOPE(FChangelistService::EnsureShelfIsCurrent);
 
 	bIsShelfReady = false;
-	if (GetShelvedFilesDepotPaths().Num() != 0)
+	if (HasShelvedFiles())
 	{
 		CheckShelfIsUpToDate([this](bool bIsUpToDate, TArray<FString> InFilesOutOfDate) {
 			// if shelf is not up to date, we need to refresh it
@@ -957,9 +973,7 @@ void FChangelistService::EnsureShelfIsCurrent()
 					if (InFilesOutOfDate.IsEmpty())
 					{
 						UE_LOG(LogSubmitToolP4, Log, TEXT("Recreating shelf"));
-						DeleteShelvedFiles(FSourceControlOperationComplete::CreateLambda([this](const FSourceControlOperationRef& Operation, ECommandResult::Type Result) {
-							CreateShelvedFiles(FSourceControlOperationComplete());
-							}));
+						ReplaceShelvedFiles();
 					}
 					else
 					{
@@ -989,7 +1003,7 @@ void FChangelistService::EnsureShelfIsCurrent()
 		if (DialogResult == EDialogFactoryResult::Confirm)
 		{
 			UE_LOG(LogSubmitToolP4, Log, TEXT("Creating shelf"));
-			CreateShelvedFiles(FSourceControlOperationComplete());
+			CreateShelvedFiles();
 		}
 	}
 }
diff --git a/Engine/Source/Programs/SubmitTool/Private/Logic/ChangelistService.h b/Engine/Source/Programs/SubmitTool/Private/Logic/ChangelistService.h
--- a/Engine/Source/Programs/SubmitTool/Private/Logic/ChangelistService.h
+++ b/Engine/Source/Programs/SubmitTool/Private/Logic/ChangelistService.h
@@ -3,19 +3,13 @@
 #pragma once
 
 #include "CoreMinimal.h"
-#include "ISourceControlProvider.h"
-#include "ISourceControlChangelist.h"
-#include "ISourceControlChangelistState.h"
-#include "Models/SCFile.h"
+#include "Logic/Services/Interfaces/ISTSourceControlService.h"
 #include "Containers/Ticker.h"
 #include "Parameters/SubmitToolParameters.h"
 #include "Services/Interfaces/ISubmitToolService.h"
 #include "Tasks/Task.h"
 
-struct FSourceControlResultInfo;
-struct FSCCommand;
 class FSubmitToolServiceProvider;
-class ISTSourceControlService;
 using FSourceControlOperationRef = TSharedRef<class ISourceControlOperation, ESPMode::ThreadSafe>;
 
 
@@ -24,6 +18,21 @@ DECLARE_DELEGATE_OneParam(FOnChangeListReadyDelegate, bool /*bValidState*/)
 DECLARE_DELEGATE_OneParam(FOnChangelistRefreshDelegate, ETaskArea /*ChangedValue*/)
 DECLARE_MULTICAST_DELEGATE(FOnCLDescriptionUpdated)
 
+struct FChangelistData
+{
+	FString Description;
+	TArray<FSCFileRef> Files;				// <Filename, DepotPath, State>
+	TArray<FSCFileRef> ShelvedFiles;		// <{}, DepotPath, State>
+	TMap<FString, FString> AllOpenFiles;	// <DepotPath, Change>
+	bool bIsValid = false;					// Changelist data could be fetched and parsed correctly
+	bool bIsShelfComplete = false;			// OpenFilesSet == ShelvedFilesSet, content may differ
+
+	static FChangelistData Parse(
+		const FString& CLID,
+		const FSCCommand& DescribeCmd,
+		const FSCCommand& OpenedCmd,
+		const TMap<FString, FString>& DepotFileToLocalFile);
+};
 
 class FChangelistService final : public ISubmitToolService
 {
@@ -35,11 +44,6 @@ public:
 	const FString GetCLID() const
 	{
 		check(IsInGameThread());
-		if(ChangelistPtr.IsValid())
-		{
-			return ChangelistPtr->GetIdentifier();
-		} 
-		
 		if(!CLID.IsEmpty())
 		{
 			return CLID;
@@ -77,25 +81,37 @@ public:
 	const TArray<FSCFileRef>& GetFilesInCL() const
 	{
 		check(IsInGameThread());
-		return FilesInCL;
+		return Changelist.Files;
 	}
 
 	const TArray<FSCFileRef>& GetShelvedFilesInCL() const
 	{
 		check(IsInGameThread());
-		return ShelvedFilesInCL;
+		return Changelist.ShelvedFiles;
 	}
 
 	bool HasShelvedFiles() const
 	{
 		check(IsInGameThread());
-		return ShelvedFilesInCL.Num() != 0;
+		return Changelist.ShelvedFiles.Num() > 0;
+	}
+
+	bool IsShelfComplete() const
+	{
+		check(IsInGameThread());
+		return Changelist.bIsShelfComplete;
+	}
+
+	const TMap<FString, FString>& GetOpenFilesInChangelists() const
+	{
+		check(IsInGameThread());
+		return Changelist.AllOpenFiles;
 	}
 
 	bool HasP4OperationsRunning() const
 	{
 		check(IsInGameThread());
-		return ActiveP4Operations.Num() > 0;
+		return !UpdateChangelistTask.IsCompleted() || SourceControlService->IsBusy();
 	}
 
 	bool IsShelfReady() const
@@ -104,27 +120,27 @@ public:
 	}
 
 	void Init();
-	void Submit(const FString& InDescriptionAddendum = TEXT(""), const FSourceControlOperationComplete& OnSubmitComplete = FSourceControlOperationComplete());
+	void Submit(const FString& InDescriptionAddendum = TEXT(""), TFunction<void(const FSCCommand&)> OnComplete = nullptr);
+
+	//void Submit(const FString& InDescriptionAddendum = TEXT(""), const FSourceControlOperationComplete& OnSubmitComplete = FSourceControlOperationComplete());
 	void CreateCLFromDefaultCL();
 	void FetchChangelistDataAsync();
-	void RevertUnchangedFilesAsync(const FSourceControlOperationComplete& OnRevertComplete = FSourceControlOperationComplete());
-	void SendCLDescriptionToP4(EConcurrency::Type Concurrency = EConcurrency::Asynchronous, FOnP4OperationCompleteDelegate InCallback = nullptr);
-	void DeleteShelvedFiles(const FSourceControlOperationComplete& OnDeleteComplete);
-	void CreateShelvedFiles(const FSourceControlOperationComplete& OnCreateComplete);
-	void UpdateShelvedFiles(TArray<FString> InFiles, TFunction<void(bool)> OnComplete = nullptr);
+	void RevertUnchangedFiles(TFunction<void(const FSCCommand&)> OnComplete = nullptr);
+	
+	void UpdateP4CLDescription();
+	void UpdateP4CLDescriptionSynchronously();
+	void DeleteShelvedFiles(TFunction<void(const FSCCommand&)> OnComplete = nullptr);
+	void ReplaceShelvedFiles(TFunction<void(const FSCCommand&)> OnComplete = nullptr);
+	void CreateShelvedFiles(TFunction<void(const FSCCommand&)> OnComplete = nullptr);
+	void UpdateShelvedFiles(TArray<FString> InFiles, TFunction<void(const FSCCommand&)> OnComplete = nullptr);
 	void CheckShelfIsUpToDate(TFunction<void(bool, TArray<FString>)> OnComplete);
 	void EnsureShelfIsCurrent();
 	bool P4Tick(float DeltaTime);
 
-	bool IsP4OperationRunning(FName OperationName = FName());
-	void CancelP4Operations(FName OperationName = FName());
+	void CancelP4Operations();
 
 	FOnCLDescriptionUpdated OnCLDescriptionUpdated;
 
-	const TArray<FString>& GetOpenedFilesDepotPaths();
-	const TArray<FString>& GetShelvedFilesDepotPaths();
-	const TMap<FString, FString>& GetOpenFilesInChangelists();
-
 private:
 	mutable FCriticalSection Mutex;
 
@@ -135,47 +151,28 @@ private:
 	const FOnChangelistRefreshDelegate& CLRefreshCallback;
 
 	const FTSTicker::FDelegateHandle TickHandle;
-	ISourceControlProvider* CachedSCCProvider;
 	const TSharedPtr<FSubmitToolServiceProvider> ServiceProvider;
 	const TSharedPtr<ISTSourceControlService> SourceControlService;
 
-	TArray<TSharedRef<ISourceControlOperation>> ActiveP4Operations;
-	UE::Tasks::TTask<TSharedRef<FSCCommand>> InitialP4OpenedFilesTask;
-	UE::Tasks::TTask<TSharedRef<FSCCommand>> InitialP4ShelvedFilesTask;
-
-	FSourceControlChangelistPtr ChangelistPtr = nullptr;
-	FSourceControlChangelistStatePtr ChangelistState = nullptr;
-	// Extracted data cached from ChangelistState by CacheChangelistState/UpdateCachedChangelistState
-	FText OriginalDescription;
+	FChangelistData Changelist;
 	FString CLDescription;
-	TArray<FSCFileRef> FilesInCL;
-	TArray<FSCFileRef> ShelvedFilesInCL;
-	// Data filled by duplicated queries in GetOpenedFilesDepotPaths/GetShelvedFilesDepotPaths since the SourceControl API hides this data
-	TArray<FString> FilesDepotPaths;
-	TArray<FString> ShelvedFilesDepotPaths;
-	// Map filled once by GetOpenedFilesDepotPaths which is more efficient than a FUpdatePendingChangelistsStatus operation for all states for all pending changelists.
-	TMap<FString, FString> OpenFilesInChangelists;
 	bool bIsShelfReady = false;
 	bool bDiffTaskRunning = false;
-	void FindInitialChangelistsAsync();
-
-	void WaitForInitialP4RunTasks();
-	UE::Tasks::TTask<TSharedRef<FSCCommand>> LaunchGetOpenedFilesDepotPaths() const;
-	UE::Tasks::TTask<TSharedRef<FSCCommand>> LaunchGetShelvedFilesDepotPaths() const;
-
-	void ParseOpenedFilesDepotPaths(TSharedRef<FSCCommand> InCompletedSCCommand);
-	void ParseShelvedFilesDepotPaths(TSharedRef<FSCCommand> InCompletedSCCommand);
-
-	void FetchChangelistCallback(const FSourceControlOperationRef& UpdateOperation, ECommandResult::Type Result);
-	void CacheChangelistState();
-	void UpdateCachedChangelistState();
-	bool AreCLDescriptionsIdentical() const;
+	TOptional<bool> BroadcastReadyIsValid;
+	TOptional<ETaskArea> BroadcastRefreshChangeType;
+
+	UE::Tasks::TTask<TSharedRef<FSCCommand>> OpenedTask;
+	UE::Tasks::TTask<TSharedRef<FSCCommand>> DescribeTask;
+	UE::Tasks::TTask<FChangelistData> ParseChangelistTask;
+	UE::Tasks::FTask UpdateChangelistTask;
+
+	void WaitForChangelistData();
+	FString ParseP4ChangesOutput(const FSCCRecord& InP4ChangeOutputRecords, const FString& InReplaceDescription);
+	void CacheChangelistData();
+	void UpdateCachedChangelistData();
 	void PrintFilesAndShelvedFiles() const;
-
-	void PrintMessages(const FSourceControlResultInfo& ResultInfo) const;
-	void PrintErrorMessages(const FSourceControlResultInfo& ResultInfo) const;
-
-
+	TArray<FSCFileRef> ParseShelvedFilesFromP4(const FSCCommand& InCmd);
+	TMap<FString, FString> RunWhereCommandSync(TArray<FString>&& InDepotPaths);
 };
 
 Expose_TNameOf(FChangelistService);
diff --git a/Engine/Source/Programs/SubmitTool/Private/Logic/IntegrationService.cpp b/Engine/Source/Programs/SubmitTool/Private/Logic/IntegrationService.cpp
--- a/Engine/Source/Programs/SubmitTool/Private/Logic/IntegrationService.cpp
+++ b/Engine/Source/Programs/SubmitTool/Private/Logic/IntegrationService.cpp
@@ -27,7 +27,8 @@ FIntegrationService::FIntegrationService
 	Parameters(InParameters),
 	ServiceProvider(InServiceProvider)
 {
-	const FString CLID = ServiceProvider.Pin()->GetService<FChangelistService>()->GetCLID();
+	FString CLID;
+	FCmdLineParameters::Get().GetValue(FSubmitToolCmdLine::P4ChangeList, CLID);
 	for(FJiraIntegrationField& Field : Parameters.Fields)
 	{
 		if(Field.Name.IsEmpty())
@@ -99,43 +100,37 @@ void FIntegrationService::RequestIntegration(const FOnBooleanValueChanged OnComp
 
 		if(!ChangelistService->GetFilesInCL().IsEmpty())
 		{
-			TFunction<void()> ShelveAndCreateSwarm = [this, ChangelistService, SourceControlService, SwarmService, TagService, JiraService, JiraTagValue, OnComplete]() {
-				ChangelistService->CreateShelvedFiles(FSourceControlOperationComplete::CreateLambda([this, ChangelistService, SourceControlService, SwarmService, TagService, JiraService, JiraTagValue, OnComplete](const FSourceControlOperationRef& DeleteShelvedOp, ECommandResult::Type Result)
-					{
-						if(Result == ECommandResult::Succeeded)
+			TFunction<void()> CreateSwarm = [this, JiraTagValue, OnComplete]() {
+				TSharedPtr<FSwarmService> SwarmService = ServiceProvider.Pin()->GetService<FSwarmService>();
+				TSharedPtr<FChangelistService> ChangelistService = ServiceProvider.Pin()->GetService<FChangelistService>();
+				TSharedPtr<ISTSourceControlService> SourceControlService = ServiceProvider.Pin()->GetService<ISTSourceControlService>();
+				TSharedPtr<FJiraService> JiraService = ServiceProvider.Pin()->GetService<FJiraService>();
+				FString SwarmUrl;
+				if (!SwarmService->GetCurrentReviewUrl(SwarmUrl))
+				{
+					SwarmService->CreateReview(SwarmService->GetUsersInSwarmTag(), OnCreateReviewComplete::CreateLambda([this, ChangelistService, SourceControlService, JiraService, JiraTagValue, OnComplete](bool bSuccess, const FString& Response) {
+						if (bSuccess)
 						{
-							FString SwarmUrl;
-							if(!SwarmService->GetCurrentReviewUrl(SwarmUrl))
-							{
-								SwarmService->CreateReview(SwarmService->GetUsersInSwarmTag(), OnCreateReviewComplete::CreateLambda([this, ChangelistService, SourceControlService, JiraService, JiraTagValue, OnComplete](bool bSuccess, const FString& Response) {
-									if(bSuccess)
-									{
-										JiraService->GetIssueAndCreateServiceDeskRequest(JiraTagValue, ChangelistService->GetCLDescription(), Response, SourceControlService->GetCurrentStreamName(), IntegrationOptions, OnComplete);
-									}
-									else
-									{
-										UE_LOG(LogSubmitTool, Error, TEXT("Failed to create swarm review, Integration request is cancelled"));
-									}
-									}));
-							}
-							else
-							{
-								SwarmService->UpdateReviewDescription(FOnBooleanValueChanged::CreateLambda([this, ChangelistService, SourceControlService, JiraService, JiraTagValue, SwarmUrl, OnComplete](bool bSuccess)
-									{
-										if(!bSuccess)
-										{
-											UE_LOG(LogSubmitTool, Warning, TEXT("Failed to update swarm review description, Integration will continue with the current swarm description"));
-										}
-
-										JiraService->GetIssueAndCreateServiceDeskRequest(JiraTagValue, ChangelistService->GetCLDescription(), SwarmUrl, SourceControlService->GetCurrentStreamName(), IntegrationOptions, OnComplete);
-									}), ChangelistService->GetCLDescription());
-							}
+							JiraService->GetIssueAndCreateServiceDeskRequest(JiraTagValue, ChangelistService->GetCLDescription(), Response, SourceControlService->GetCurrentStreamName(), IntegrationOptions, OnComplete);
 						}
 						else
 						{
-							UE_LOG(LogSubmitTool, Error, TEXT("Failed to shelve files, Integration request is cancelled"));
+							UE_LOG(LogSubmitTool, Error, TEXT("Failed to create swarm review, Integration request is cancelled"));
 						}
-					}));
+						}));
+				}
+				else
+				{
+					SwarmService->UpdateReviewDescription(FOnBooleanValueChanged::CreateLambda([this, ChangelistService, SourceControlService, JiraService, JiraTagValue, SwarmUrl, OnComplete](bool bSuccess)
+						{
+							if (!bSuccess)
+							{
+								UE_LOG(LogSubmitTool, Warning, TEXT("Failed to update swarm review description, Integration will continue with the current swarm description"));
+							}
+
+							JiraService->GetIssueAndCreateServiceDeskRequest(JiraTagValue, ChangelistService->GetCLDescription(), SwarmUrl, SourceControlService->GetCurrentStreamName(), IntegrationOptions, OnComplete);
+						}), ChangelistService->GetCLDescription());
+				}
 			};
 
 			if(ChangelistService->HasShelvedFiles())
@@ -143,22 +138,32 @@ void FIntegrationService::RequestIntegration(const FOnBooleanValueChanged OnComp
 				EDialogFactoryResult Result = FDialogFactory::ShowConfirmDialog(FText::FromString(TEXT("Shelve files")), FText::FromString(TEXT("Submit tool will shelve your local files to update or create the swarm review.\nThis will replace the existing shelve entirely so you could lose any shelved changes that are not local, do you want to continue?")));
 				if(Result == EDialogFactoryResult::Confirm)
 				{
-					ChangelistService->DeleteShelvedFiles(FSourceControlOperationComplete::CreateLambda([this, ShelveAndCreateSwarm](const FSourceControlOperationRef& DeleteShelvedOp, ECommandResult::Type Result)
+					ChangelistService->ReplaceShelvedFiles([this, CreateSwarm](const FSCCommand& InCmd)
+					{
+						if(InCmd.bSuccess)
 						{
-							if(Result == ECommandResult::Succeeded)
-							{
-								ShelveAndCreateSwarm();
-							}
-							else
-							{
-								UE_LOG(LogSubmitTool, Error, TEXT("Failed to delete shelve, Integration request is cancelled"));
-							}
-						}));
+							CreateSwarm();
+						}
+						else
+						{
+							UE_LOG(LogSubmitTool, Error, TEXT("Failed to replace shelve, Integration request is cancelled"));
+						}
+					});
 				}
 			}
 			else
 			{
-				ShelveAndCreateSwarm();
+				ChangelistService->CreateShelvedFiles([this, CreateSwarm](const FSCCommand& InCmd)
+				{
+					if (InCmd.bSuccess)
+					{
+						CreateSwarm();
+					}
+					else
+					{
+						UE_LOG(LogSubmitTool, Error, TEXT("Failed to shelve files, Integration request is cancelled"));
+					}
+				});
 			}
 		}
 		else if(!ChangelistService->GetShelvedFilesInCL().IsEmpty())
diff --git a/Engine/Source/Programs/SubmitTool/Private/Logic/JiraService.cpp b/Engine/Source/Programs/SubmitTool/Private/Logic/JiraService.cpp
--- a/Engine/Source/Programs/SubmitTool/Private/Logic/JiraService.cpp
+++ b/Engine/Source/Programs/SubmitTool/Private/Logic/JiraService.cpp
@@ -68,7 +68,7 @@ void FJiraService::QueryIssues()
 		return;
 	}
 
-	FHttpModule& HttpModule = FModuleManager::LoadModuleChecked<FHttpModule>("HTTP");
+	FHttpModule& HttpModule = FModuleManager::GetModuleChecked<FHttpModule>("HTTP");
 	JiraRequest = HttpModule.Get().CreateRequest();
 	
 	JiraRequest->OnProcessRequestComplete().BindRaw(this, &FJiraService::QueryIssues_HttpRequestComplete);
diff --git a/Engine/Source/Programs/SubmitTool/Private/Logic/P4LockdownService.cpp b/Engine/Source/Programs/SubmitTool/Private/Logic/P4LockdownService.cpp
--- a/Engine/Source/Programs/SubmitTool/Private/Logic/P4LockdownService.cpp
+++ b/Engine/Source/Programs/SubmitTool/Private/Logic/P4LockdownService.cpp
@@ -88,7 +88,7 @@ bool FP4LockdownService::WaitForAllowListData()
 	return bSuccess;
 }
 
-FPathLockdownResult FP4LockdownService::ArePathsInLockdown(const TArray<FString>& InPaths)
+FPathLockdownResult FP4LockdownService::ArePathsInLockdown(const TArray<FSCFileRef>& InPaths)
 {	
 	TRACE_CPUPROFILER_EVENT_SCOPE(FP4LockdownService::ArePathsInLockdown);
 
@@ -114,13 +114,13 @@ FPathLockdownResult FP4LockdownService::ArePathsInLockdown(const TArray<FString>
 	}
 
 	// Finally validate the paths
-	for(const FString& Path : InPaths)
+	for(const FSCFileRef& Path : InPaths)
 	{
-		FPathLockdownResult Result = IsPathInLockdown(Path);
+		FPathLockdownResult Result = IsPathInLockdown(Path->GetDepotPath());
 
 		if(Result.bIsLocked)
 		{
-			UE_LOG(LogSubmitToolP4Debug, Warning, TEXT("File %s is locked down by %s.\n%s"), *Path, *FSubmitToolCoreUtils::StringBuilderJoin<64>(Result.GroupNames, TEXT(", ")), *FSubmitToolCoreUtils::StringBuilderJoin(Result.Messages, TEXT("\n")));
+			UE_LOG(LogSubmitToolP4Debug, Warning, TEXT("File %s is locked down by %s.\n%s"), *Path->GetDepotPath(), *FSubmitToolCoreUtils::StringBuilderJoin<64>(Result.GroupNames, TEXT(", ")), *FSubmitToolCoreUtils::StringBuilderJoin(Result.Messages, TEXT("\n")));
 			GlobalResult.bAllowlisted &= Result.bAllowlisted;
 			GlobalResult.bIsLocked = Result.bIsLocked;
 
diff --git a/Engine/Source/Programs/SubmitTool/Private/Logic/P4LockdownService.h b/Engine/Source/Programs/SubmitTool/Private/Logic/P4LockdownService.h
--- a/Engine/Source/Programs/SubmitTool/Private/Logic/P4LockdownService.h
+++ b/Engine/Source/Programs/SubmitTool/Private/Logic/P4LockdownService.h
@@ -28,7 +28,7 @@ public:
 
 	FP4LockdownService() = delete;
 	FP4LockdownService(const FP4LockdownParameters& InParameters, TWeakPtr<FSubmitToolServiceProvider> InServiceProvider);
-	FPathLockdownResult ArePathsInLockdown(const TArray<FString>& InPaths);
+	FPathLockdownResult ArePathsInLockdown(const TArray<FSCFileRef>& InPaths);
 	bool IsBlockingOperationRunning() const
 	{
 		return DownloadTask.IsValid() && !DownloadTask.IsCompleted();
diff --git a/Engine/Source/Programs/SubmitTool/Private/Logic/PreflightService.cpp b/Engine/Source/Programs/SubmitTool/Private/Logic/PreflightService.cpp
--- a/Engine/Source/Programs/SubmitTool/Private/Logic/PreflightService.cpp
+++ b/Engine/Source/Programs/SubmitTool/Private/Logic/PreflightService.cpp
@@ -507,9 +507,9 @@ void FPreflightService::FetchUnlinkedPreflight(const FString& InPreflightId, boo
 
 void FPreflightService::StartPreflight()
 {
-	const TArray<FString>& ShelvedPaths = ServiceProvider.Pin()->GetService<FChangelistService>()->GetShelvedFilesDepotPaths();
+	const TArray<FSCFileRef>& ShelvedFiles = ServiceProvider.Pin()->GetService<FChangelistService>()->GetShelvedFilesInCL();
 
-	if (ShelvedPaths.IsEmpty())
+	if (ShelvedFiles.IsEmpty())
 	{
 		LastErrorMessage = TEXT("Shelve is empty or it couldn't be retrieved from p4, can't request preflight");
 		State = EPreflightServiceState::Error;
@@ -598,16 +598,16 @@ bool FPreflightService::SelectPreflightTemplate(FPreflightTemplateDefinition& Ou
 	TSharedPtr<FChangelistService> ChangelistService = ServiceProvider.Pin()->GetService<FChangelistService>();
 
 	// Get the list of files in the changelist
-	const TArray<FString>& FilesInCl = bCheckShelveInstead ? ChangelistService->GetShelvedFilesDepotPaths() : ChangelistService->GetOpenedFilesDepotPaths();
+	const TArray<FSCFileRef>& FilesInCl = bCheckShelveInstead ? ChangelistService->GetShelvedFilesInCL() : ChangelistService->GetFilesInCL();
 
 	// Loop through each definition to see if the files are in the path then check extension
 	for(const FPreflightTemplateDefinition& Def : Definition.Definitions)
 	{
 		FString RegexPat = Def.RegexPath.Replace(TEXT("$(StreamRoot)"), *StreamName, ESearchCase::IgnoreCase);
 		FRegexPattern Pattern = FRegexPattern(RegexPat, ERegexPatternFlags::CaseInsensitive);
-		for(const FString& File : FilesInCl)
+		for(const FSCFileRef& File : FilesInCl)
 		{
-			FRegexMatcher regex = FRegexMatcher(Pattern, File);
+			FRegexMatcher regex = FRegexMatcher(Pattern, File->GetDepotPath());
 			if(regex.FindNext())
 			{
 				OutTemplate = Def;
@@ -626,14 +626,14 @@ FString FPreflightService::GetAdditionalTasksString(const FPreflightTemplateDefi
 	const FString EndString = TEXT("=true");
 	TSharedPtr<FChangelistService> ChangelistService = ServiceProvider.Pin()->GetService<FChangelistService>();
 
-	const TArray<FString>& FilesInCl = bCheckShelveInstead ? ChangelistService->GetShelvedFilesDepotPaths() : ChangelistService->GetOpenedFilesDepotPaths();
+	const TArray<FSCFileRef>& FilesInCl = bCheckShelveInstead ? ChangelistService->GetShelvedFilesInCL() : ChangelistService->GetFilesInCL();
 	for(const FPreflightAdditionalTask& AdditionalTask : InTemplate.AdditionalTasks)
 	{
 		FString RegexPat = AdditionalTask.RegexPath.Replace(TEXT("$(StreamRoot)"), *StreamName, ESearchCase::IgnoreCase);
 		FRegexPattern Pattern = FRegexPattern(RegexPat, ERegexPatternFlags::CaseInsensitive);
-		for(const FString& File : FilesInCl)
+		for(const FSCFileRef& File : FilesInCl)
 		{
-			FRegexMatcher regex = FRegexMatcher(Pattern, File);
+			FRegexMatcher regex = FRegexMatcher(Pattern, File->GetDepotPath());
 			if(regex.FindNext())
 			{
 				AdditionalTaskStrBuilder.Append(BaseString);
@@ -651,12 +651,12 @@ void FPreflightService::RefreshStreamName()
 	TSharedPtr<ISTSourceControlService> SCCService = ServiceProvider.Pin()->GetService<ISTSourceControlService>();
 	StreamName = SCCService->GetRootStreamName();
 
-	const TArray<FString>& ShelvedPaths = ServiceProvider.Pin()->GetService<FChangelistService>()->GetShelvedFilesDepotPaths();
+	const TArray<FSCFileRef>& ShelvedFiles = ServiceProvider.Pin()->GetService<FChangelistService>()->GetShelvedFilesInCL();
 
-	if (ShelvedPaths.Num() != 0)
+	if (ShelvedFiles.Num() != 0)
 	{
-		FString CommonPath = ShelvedPaths[0];
-		const FString& LastPath = ShelvedPaths.Last();
+		FString CommonPath = ShelvedFiles[0]->GetDepotPath();
+		const FString& LastPath = ShelvedFiles.Last()->GetDepotPath();
 
 		for (size_t i = 0; i < FMath::Min(CommonPath.Len(), LastPath.Len()); ++i)
 		{
diff --git a/Engine/Source/Programs/SubmitTool/Private/Logic/Services/CacheDataService.cpp b/Engine/Source/Programs/SubmitTool/Private/Logic/Services/CacheDataService.cpp
--- a/Engine/Source/Programs/SubmitTool/Private/Logic/Services/CacheDataService.cpp
+++ b/Engine/Source/Programs/SubmitTool/Private/Logic/Services/CacheDataService.cpp
@@ -31,12 +31,12 @@ const FDateTime FCacheDataService::GetLastValidationDate(const FString& InCLId,
 	return Data.GetLastValidationDate(InFilePath, InValidatorId, InValidatorConfig);
 }
 
-void FCacheDataService::UpdateLastValidationForFiles(const FString& InCLId, const FName& InValidatorId, const FString& InValidatorConfig, const TArray<FSourceControlStateRef>& InFiles, const FDateTime& InNewTimestamp)
+void FCacheDataService::UpdateLastValidationForFiles(const FString& InCLId, const FName& InValidatorId, const FString& InValidatorConfig, const TArray<FSCFileRef>& InFiles, const FDateTime& InNewTimestamp)
 {
 	FChangelistCacheData& Data = CacheData->CLCacheData.FindOrAdd(InCLId);
 	Data.LastAccessed = FDateTime::UtcNow();
 
-	for(const FSourceControlStateRef& File : InFiles)
+	for(const FSCFileRef& File : InFiles)
 	{
 		FValidationRecord& ValidationRecord = Data.LastFileValidations.FindOrAdd(File->GetFilename()).LastValidationDateTimes.FindOrAdd(InValidatorId);
 		ValidationRecord.ValidatorConfig = InValidatorConfig;
diff --git a/Engine/Source/Programs/SubmitTool/Private/Logic/Services/CacheDataService.h b/Engine/Source/Programs/SubmitTool/Private/Logic/Services/CacheDataService.h
--- a/Engine/Source/Programs/SubmitTool/Private/Logic/Services/CacheDataService.h
+++ b/Engine/Source/Programs/SubmitTool/Private/Logic/Services/CacheDataService.h
@@ -4,7 +4,6 @@
 
 #include "CoreMinimal.h"
 #include "Models/CacheData.h"
-#include "ISourceControlState.h"
 #include "Interfaces/ICacheDataService.h"
 
 struct FGeneralParameters;
@@ -15,7 +14,7 @@ public:
 	virtual bool GetChangelistCacheData(const FString& InCLId, FChangelistCacheData& OutCLData) const override { return false; }
 
 	virtual const FDateTime GetLastValidationDate(const FString& InCLId, const FName& InValidatorId, const FString& InValidatorConfig, const FString& InDepotPath) const override { return FDateTime::MinValue(); }
-	virtual void UpdateLastValidationForFiles(const FString& InCLId, const FName& InValidatorId, const FString& InValidatorConfig, const TArray<FSourceControlStateRef>& InFiles, const FDateTime& InNewTimestamp) override {}
+	virtual void UpdateLastValidationForFiles(const FString& InCLId, const FName& InValidatorId, const FString& InValidatorConfig, const TArray<FSCFileRef>& InFiles, const FDateTime& InNewTimestamp) override {}
 
 	virtual FString GetIntegrationFieldValue(const FString& InCLId, const FString& InIntegrationFieldId) const override { return FString(); }
 	virtual void SetIntegrationFieldValue(const FString& InCLId, const FString& InIntegrationFieldId, const FString& InValue) { }
@@ -35,7 +34,7 @@ public:
 	virtual bool GetChangelistCacheData(const FString& InCLId, FChangelistCacheData& OutCLData) const override;
 
 	virtual const FDateTime GetLastValidationDate(const FString& InCLId, const FName& InValidatorId, const FString& InValidatorConfig, const FString& InDepotPath) const override;
-	virtual void UpdateLastValidationForFiles(const FString& InCLId, const FName& InValidatorId, const FString& InValidatorConfig, const TArray<FSourceControlStateRef>& InFiles, const FDateTime& InNewTimestamp) override;
+	virtual void UpdateLastValidationForFiles(const FString& InCLId, const FName& InValidatorId, const FString& InValidatorConfig, const TArray<FSCFileRef>& InFiles, const FDateTime& InNewTimestamp) override;
 
 	virtual FString GetIntegrationFieldValue(const FString& InCLId, const FString& InIntegrationFieldId) const override;
 	virtual void SetIntegrationFieldValue(const FString& InCLId, const FString& InIntegrationFieldId, const FString& InValue);
@@ -52,4 +51,4 @@ private:
 
 	const FGeneralParameters& Parameters;
 	TUniquePtr<FCacheFile> CacheData;
-};
\ No newline at end of file
+};
diff --git a/Engine/Source/Programs/SubmitTool/Private/Logic/Services/Interfaces/ICacheDataService.h b/Engine/Source/Programs/SubmitTool/Private/Logic/Services/Interfaces/ICacheDataService.h
--- a/Engine/Source/Programs/SubmitTool/Private/Logic/Services/Interfaces/ICacheDataService.h
+++ b/Engine/Source/Programs/SubmitTool/Private/Logic/Services/Interfaces/ICacheDataService.h
@@ -5,7 +5,7 @@
 #include "CoreMinimal.h"
 #include "ISubmitToolService.h"
 #include "Models/CacheData.h"
-#include "ISourceControlState.h"
+#include "Models/SCFile.h"
 
 struct FGeneralParameters;
 
@@ -15,7 +15,7 @@ public:
 	virtual bool GetChangelistCacheData(const FString& InCLId, FChangelistCacheData& OutCLData) const = 0;
 
 	virtual const FDateTime GetLastValidationDate(const FString& InCLId, const FName& InValidatorId, const FString& InValidatorConfig, const FString& InDepotPath) const = 0;
-	virtual void UpdateLastValidationForFiles(const FString& InCLId, const FName& InValidatorId, const FString& InValidatorConfig, const TArray<FSourceControlStateRef>& InFiles, const FDateTime& InNewTimestamp) = 0;
+	virtual void UpdateLastValidationForFiles(const FString& InCLId, const FName& InValidatorId, const FString& InValidatorConfig, const TArray<FSCFileRef>& InFiles, const FDateTime& InNewTimestamp) = 0;
 
 	virtual FString GetIntegrationFieldValue(const FString& InCLId, const FString& InIntegrationFieldId) const = 0;
 	virtual void SetIntegrationFieldValue(const FString& InCLId, const FString& InIntegrationFieldId, const FString& InValue) = 0;
diff --git a/Engine/Source/Programs/SubmitTool/Private/Logic/Services/Interfaces/ISTSourceControlService.h b/Engine/Source/Programs/SubmitTool/Private/Logic/Services/Interfaces/ISTSourceControlService.h
--- a/Engine/Source/Programs/SubmitTool/Private/Logic/Services/Interfaces/ISTSourceControlService.h
+++ b/Engine/Source/Programs/SubmitTool/Private/Logic/Services/Interfaces/ISTSourceControlService.h
@@ -4,7 +4,6 @@
 
 #include "CoreMinimal.h"
 #include "ISubmitToolService.h"
-#include "ISourceControlProvider.h"
 #include "Parameters/SubmitToolParameters.h"
 #include "Memory/SharedBuffer.h"
 #include "Tasks/Task.h"
@@ -62,6 +61,27 @@ struct FUserData
 	}
 };
 
+struct FSCResultInfo
+{
+	/** Append any messages from another FSCResultInfo, ensuring to keep any already accumulated info. */
+	void Append(const FSCResultInfo& InResultInfo)
+	{
+		InfoMessages.Append(InResultInfo.InfoMessages);
+		ErrorMessages.Append(InResultInfo.ErrorMessages);
+	}
+
+	bool HasErrors() const
+	{
+		return !ErrorMessages.IsEmpty();
+	}
+
+	/** Info and/or warning message storage */
+	TArray<FText> InfoMessages;
+
+	/** Potential error message storage */
+	TArray<FText> ErrorMessages;
+};
+
 
 struct FSCCStream
 {
@@ -76,7 +96,8 @@ public:
 	TArray<const FString> AdditionalImportPaths;
 };
 
-using FSCCRecordSet = TArray<TMap<FString, FString>>;
+using FSCCRecord = TMap<FString, FString>;
+using FSCCRecordSet = TArray<FSCCRecord>;
 
 struct FSCCommand
 {
@@ -85,11 +106,13 @@ struct FSCCommand
 	TArray<FString> Args;
 	bool bCollectData = false;
 	bool bSilentErrors = false;
+	FString InputData;
 	// out
 	FSCCRecordSet Values;
-	FSourceControlResultInfo Results;
+	FSCResultInfo Results;
 	TMap<FString, FSharedBuffer> Data;
 	bool bSuccess = false;
+	bool bDroppedConnection = false;
 };
 
 DECLARE_MULTICAST_DELEGATE_OneParam(FOnUsersGet, TArray<TSharedPtr<FUserData>>& /*Users*/)
@@ -102,17 +125,11 @@ class ISTSourceControlService : public ISubmitToolService
 {
 public:
 	virtual void InitializeParameters(const FSubmitToolParameters& InParameters) = 0;
-	virtual const TUniquePtr<ISourceControlProvider>& GetProvider() const = 0;
 	virtual void GetUsers(const FOnUsersGet::FDelegate& Callback) = 0;
 	virtual void GetGroups(const FOnGroupsGet::FDelegate& Callback) = 0;
 	virtual void GetUsersAndGroups(const FOnUsersAndGroupsGet::FDelegate& Callback) = 0;
 
 	virtual UE::Tasks::TTask<bool> DownloadFiles(TArray<FString>&& FilePaths, TMap<FString, FSharedBuffer>& OutFileBuffers) = 0;
-
-	virtual bool IsAvailable() const = 0;
-
-	virtual bool Tick(float InDeltaTime) = 0;
-
 	virtual const TArray<TSharedPtr<FUserData>>& GetRecentUsers() const = 0;
 	virtual void AddRecentUser(TSharedPtr<FUserData>& User) = 0;
 	virtual const TArray<TSharedPtr<FString>>& GetRecentGroups() const = 0;
@@ -126,14 +143,15 @@ public:
 	virtual const FString GetCurrentStreamName() = 0;
 	virtual int32 GetDepotStreamDepth(const FString& InDepotName) = 0;
 	virtual const FAuthTicket& GetAuthTicket() = 0;
+	virtual void RequestCancel() = 0;
 	virtual bool IsBusy() const = 0;
 
-	// Runs the command in a task thread and the provided callback in a dependent GT task.
-	void RunCommand(FString&& InCommand, TArray<FString>&& InAdditionalArgs, FOnSCCCommandComplete InCompleteCallback, bool bSilentErrors = false)
+	void RunCommand(FSCCommand&& InCommand, FOnSCCCommandComplete InCompleteCallback)
 	{
 		check(InCompleteCallback.IsBound());
-		TSharedRef<FSCCommand> Command = MakeShared<FSCCommand>(MoveTemp(InCommand), MoveTemp(InAdditionalArgs), false, bSilentErrors);
-		UE::Tasks::FTask CommandTask = UE::Tasks::Launch(UE_SOURCE_LOCATION, [this, Command, InCompleteCallback, bSilentErrors]
+		TSharedRef<FSCCommand> Command = MakeShared<FSCCommand>(MoveTemp(InCommand));
+
+		UE::Tasks::FTask CommandTask = UE::Tasks::Launch(UE_SOURCE_LOCATION, [this, Command, InCompleteCallback]
 		{
 			Command->bSuccess = RunCommandInternal(*Command);
 			UE::Tasks::Launch(UE_SOURCE_LOCATION, [InCompleteCallback, Command]
@@ -144,6 +162,13 @@ public:
 		});
 	}
 
+	// Runs the command in a task thread and the provided callback in a dependent GT task.
+	void RunCommand(FString&& InCommand, TArray<FString>&& InAdditionalArgs, FOnSCCCommandComplete InCompleteCallback, bool bSilentErrors = false)
+	{
+		check(InCompleteCallback.IsBound());
+		RunCommand(FSCCommand(MoveTemp(InCommand), MoveTemp(InAdditionalArgs), false, bSilentErrors), InCompleteCallback);
+	}
+
 	// Runs the command in a task thread, calling GetResult() will wait until the task is completed and return the command data.
 	UE::Tasks::TTask<TSharedRef<FSCCommand>> RunCommand(FString&& InCommand, TArray<FString>&& InAdditionalArgs, bool bSilentErrors = false)
 	{
@@ -163,6 +188,10 @@ public:
 		return MoveTemp(Command);
 	}
 
+	void RunSyncCommand(FSCCommand& InCommand)
+	{
+		InCommand.bSuccess = RunCommandInternal(InCommand);
+	}
 protected:
 	virtual bool RunCommandInternal(FSCCommand& InCommand) = 0;
 
diff --git a/Engine/Source/Programs/SubmitTool/Private/Logic/Services/SourceControl/PerforceData.h b/Engine/Source/Programs/SubmitTool/Private/Logic/Services/SourceControl/PerforceData.h
--- a/Engine/Source/Programs/SubmitTool/Private/Logic/Services/SourceControl/PerforceData.h
+++ b/Engine/Source/Programs/SubmitTool/Private/Logic/Services/SourceControl/PerforceData.h
@@ -2,7 +2,6 @@
 
 #pragma once
 #include "Memory/SharedBuffer.h"
-#include "SourceControlResultInfo.h"
 #include "Internationalization/Internationalization.h"
 #include "Logic/Services/Interfaces/ISTSourceControlService.h"
 
@@ -153,15 +152,25 @@ class FSTClientUser : public ClientUser
 {
 public:
 
-	FSTClientUser(FSCCRecordSet& InRecords, EP4ClientUserFlags InFlags, FSourceControlResultInfo& OutResultInfo)
+	FSTClientUser(FSCCommand& InCommand, EP4ClientUserFlags InFlags)
 		: ClientUser()
 		, Flags(InFlags)
-		, Records(InRecords)
-		, ResultInfo(OutResultInfo)
+		, Command(InCommand)
 	{
 
 	}
 
+
+	/** Called by P4API on "change -i" command to fill OutBuffer. */
+	virtual void InputData(StrBuf* OutBuffer, Error* OutError) override
+	{
+		if (!Command.InputData.IsEmpty())
+		{
+			OutBuffer->Append(FROM_TCHAR(*Command.InputData, IsUnicodeServer()));
+		}
+	}
+
+
 	/**  Called by P4API when the results from running a command are ready. */
 	virtual void OutputStat(StrDict* VarList) override
 	{
@@ -183,13 +192,13 @@ public:
 				FText Message = FText::Format(LOCTEXT("P4Client_GatheringUnfinished", "Started gathering depot file '{0}' before the previous file finished!"),
 					FText::FromString(File.GetDepotPath()));
 
-				ResultInfo.ErrorMessages.Add(MoveTemp(Message));
+				Command.Results.ErrorMessages.Add(MoveTemp(Message));
 			}
 
 			File.Initialize(Record);
 		}
 
-		Records.Add(Record);
+		Command.Values.Add(Record);
 	}
 
 	/** Called by P4API when it output a chunk of text data from a file (commonly via P4 Print) */
@@ -215,7 +224,7 @@ public:
 					File.GetRemainingBytes(),
 					FText::FromString(File.GetDepotPath()));
 
-				ResultInfo.ErrorMessages.Add(MoveTemp(Message));
+				Command.Results.ErrorMessages.Add(MoveTemp(Message));
 
 				File.Reset();
 			}
@@ -249,7 +258,7 @@ public:
 						File.GetRemainingBytes(),
 						FText::FromString(File.GetDepotPath()));
 
-					ResultInfo.ErrorMessages.Add(MoveTemp(Message));
+					Command.Results.ErrorMessages.Add(MoveTemp(Message));
 
 					File.Reset();
 				}
@@ -262,7 +271,7 @@ public:
 					File.GetRemainingBytes(),
 					FText::FromString(File.GetDepotPath()));
 
-				ResultInfo.ErrorMessages.Add(MoveTemp(Message));
+				Command.Results.ErrorMessages.Add(MoveTemp(Message));
 
 				File.Reset();
 			}
@@ -291,11 +300,11 @@ public:
 
 		if(err->GetSeverity() <= ErrorSeverity::E_INFO)
 		{
-			ResultInfo.InfoMessages.Add(FText::FromString(MoveTemp(Message)));
+			Command.Results.InfoMessages.Add(FText::FromString(MoveTemp(Message)));
 		}
 		else
 		{
-			ResultInfo.ErrorMessages.Add(FText::FromString(MoveTemp(Message)));
+			Command.Results.ErrorMessages.Add(FText::FromString(MoveTemp(Message)));
 		}
 	}
 
@@ -306,7 +315,7 @@ public:
 		// We don't expect this to ever be called (info messages should come
 		// via ClientUser::Message) but implemented just to be safe.
 
-		ResultInfo.InfoMessages.Add(FText::FromString(FString(TO_TCHAR(InInfo, IsUnicodeServer()))));
+		Command.Results.InfoMessages.Add(FText::FromString(FString(TO_TCHAR(InInfo, IsUnicodeServer()))));
 	}
 
 	virtual void OutputError(const char* errBuf) override
@@ -319,7 +328,7 @@ public:
 		// to implement this method to make sure we capture all of the errors being passed in
 		// this way.
 
-		ResultInfo.ErrorMessages.Add(FText::FromString(FString(TO_TCHAR(errBuf, IsUnicodeServer()))));
+		Command.Results.ErrorMessages.Add(FText::FromString(FString(TO_TCHAR(errBuf, IsUnicodeServer()))));
 	}
 
 	inline bool IsUnicodeServer() const
@@ -339,9 +348,8 @@ public:
 	}
 
 	EP4ClientUserFlags Flags;
-	FSCCRecordSet& Records;
 
-	FSourceControlResultInfo& ResultInfo;
+	FSCCommand& Command;
 
 private:
 	TMap<FString, FSharedBuffer> Files;
diff --git a/Engine/Source/Programs/SubmitTool/Private/Logic/Services/SourceControl/SubmitToolPerforce.cpp b/Engine/Source/Programs/SubmitTool/Private/Logic/Services/SourceControl/SubmitToolPerforce.cpp
--- a/Engine/Source/Programs/SubmitTool/Private/Logic/Services/SourceControl/SubmitToolPerforce.cpp
+++ b/Engine/Source/Programs/SubmitTool/Private/Logic/Services/SourceControl/SubmitToolPerforce.cpp
@@ -1,8 +1,6 @@
 // Copyright Epic Games, Inc. All Rights Reserved.
 
 #include "SubmitToolPerforce.h"
-#include "SourceControlInitSettings.h"
-#include "ISourceControlModule.h"
 #include "Logic/ProcessWrapper.h"
 #include "SubmitToolCoreUtils.h"
 #include "Configuration/Configuration.h"
@@ -37,9 +35,9 @@ FP4Connection::~FP4Connection()
 
 FConnectionPool::~FConnectionPool()
 {
-	RequestCancel();
 	for(TUniquePtr<FClientApiWrapper>& ConWrapper : P4Connections)
 	{
+		ConWrapper->bWantsCancel = true;
 		Error P4Error;
 		ConWrapper->Connection->Final(&P4Error);
 		if(P4Error.Test())
@@ -55,29 +53,56 @@ FConnectionPool::~FConnectionPool()
 TUniquePtr<FP4Connection> FConnectionPool::GetAvailableConnection()
 {
 	FScopeLock Lock(&Mutex);
-	for(TUniquePtr<FClientApiWrapper>& ConWrapper : P4Connections)
+
+	int32 ClosingConnections = 0;
+	int32 I = 0;
+	for (TArray<TUniquePtr<FClientApiWrapper>>::TIterator It = P4Connections.CreateIterator(); It; ++It, ++I)
 	{
-		if(ConWrapper->bIsReady)
+		FClientApiWrapper* ConWrapper = It->Get();
+		if (ConWrapper->Connection->Dropped() == 1)
+		{
+			UE_LOG(LogSubmitToolP4Debug, Log, TEXT("Connection #%d dropped."), I);
+			It.RemoveCurrent();
+		}
+		else if (ConWrapper->bWantsCancel)
+		{
+			UE_LOG(LogSubmitToolP4Debug, Log, TEXT("Connection #%d is cancelling."), I);
+			ClosingConnections++;
+		}
+		else if (ConWrapper->bIsReady)
 		{
 			return MakeUnique<FP4Connection>(*ConWrapper);
 		}
 	}
 
-
-	if(P4Connections.Num() < MaxConnections)
+	const int32 NumConnections = P4Connections.Num() - ClosingConnections;
+	if (NumConnections < MaxConnections)
 	{
-		UE_LOG(LogSubmitToolP4Debug, Log, TEXT("Creating new p4 connection %d/%d."), P4Connections.Num()+1, MaxConnections);
-		size_t Idx = CreateConnection();
-		if(Idx != -1)
+		UE_LOG(LogSubmitToolP4Debug, Log, TEXT("Creating new p4 connection %d/%d."), NumConnections + 1, MaxConnections);
+		int32 NewIdx = CreateConnection();
+		if (NewIdx != -1)
 		{
-			return MakeUnique<FP4Connection>(*P4Connections[Idx]);
+			return MakeUnique<FP4Connection>(*P4Connections[NewIdx]);
 		}
 	}
 
 	return nullptr;
 }
 
-size_t FConnectionPool::CreateConnection()
+void FConnectionPool::CancelCurrentRequests()
+{
+	for (int32 I = 0; I < P4Connections.Num(); ++I)
+	{
+		const TUniquePtr<FClientApiWrapper>& Connection = P4Connections[I];
+		if (!Connection->bIsReady)
+		{
+			UE_LOG(LogSubmitToolP4Debug, Log, TEXT("Cancelling connection #%d."), I);
+			Connection->bWantsCancel = true;
+		}
+	}
+}
+
+int32 FConnectionPool::CreateConnection()
 {
 	if(bConnectionFailed)
 	{
@@ -102,55 +127,12 @@ size_t FConnectionPool::CreateConnection()
 	}
 
 	FScopeLock Lock(&Mutex);
-	return P4Connections.Add(MakeUnique<FClientApiWrapper>(MoveTemp(P4Client), FOnIsCancelled::CreateLambda([this] { return bWantsCancel.load(); })));
+	return P4Connections.Add(MakeUnique<FClientApiWrapper>(MoveTemp(P4Client)));
 }
 
 FSubmitToolPerforce::FSubmitToolPerforce()
 {
 	TRACE_CPUPROFILER_EVENT_SCOPE(FSubmitToolPerforce::FSubmitToolPerforce);
-
-	FModuleManager::LoadModuleChecked<ISourceControlModule>(FName("PerforceSourceControl"));
-
-	FSourceControlInitSettings SCCSettings(FSourceControlInitSettings::EBehavior::OverrideAll);
-
-	const FSourceControlInitSettings::EConfigBehavior IniBehavior = FSourceControlInitSettings::EConfigBehavior::ReadOnly;
-	SCCSettings.SetConfigBehavior(IniBehavior);
-
-	FString PerforceServerAndPort;
-	FCmdLineParameters::Get().GetValue(FSubmitToolCmdLine::P4Server, PerforceServerAndPort);
-
-	FString PerforceUserName;
-	FCmdLineParameters::Get().GetValue(FSubmitToolCmdLine::P4User, PerforceUserName);
-
-	FString PerforceClientName;
-	FCmdLineParameters::Get().GetValue(FSubmitToolCmdLine::P4Client, PerforceClientName);
-
-	SCCSettings.AddSetting(TEXT("P4Port"), PerforceServerAndPort);
-	SCCSettings.AddSetting(TEXT("P4User"), PerforceUserName);
-	SCCSettings.AddSetting(TEXT("P4Client"), PerforceClientName);
-
-	SCCProvider = ISourceControlModule::Get().CreateProvider(FName("Perforce"), TEXT("SubmitTool"), SCCSettings);
-
-	if(SCCProvider.IsValid())
-	{
-		UE_LOG(LogSubmitTool, Log, TEXT("Setting Perforce Connection parameters: '%s' | User: '%s' | Workspace: '%s'"), *PerforceServerAndPort, *PerforceUserName, *PerforceClientName);
-
-		ISourceControlProvider::FInitResult ConnectionResult = SCCProvider->Init(ISourceControlProvider::EInitFlags::None);
-		if(!ConnectionResult.bIsAvailable)
-		{
-			UE_LOG(LogSubmitTool, Error, TEXT("%s"), *ConnectionResult.Errors.ErrorMessage.ToString());
-			for(FText& AdditionalErrorMessage : ConnectionResult.Errors.AdditionalErrors)
-			{
-				UE_LOG(LogSubmitTool, Error, TEXT("%s"), *AdditionalErrorMessage.ToString());
-			}
-		}
-	}
-	else
-	{
-		UE_LOG(LogSubmitTool, Error, TEXT("Failed to create a perforce revision control provider"));
-	}
-
-	TickHandle = FTSTicker::GetCoreTicker().AddTicker(FTickerDelegate::CreateRaw(this, &FSubmitToolPerforce::Tick));
 	InitialStreamTask = GetStream();
 }
 
@@ -159,11 +141,6 @@ FSubmitToolPerforce::~FSubmitToolPerforce()
 	OnUsersGetCallbacks.Clear();
 	OnGroupsGetCallbacks.Clear();
 
-	if(SCCProvider.IsValid())
-	{
-		SCCProvider->Close();
-	}
-
 	FTSTicker::RemoveTicker(TickHandle);
 }
 
@@ -172,11 +149,6 @@ void FSubmitToolPerforce::InitializeParameters(const FSubmitToolParameters& InPa
 	GroupsToExclude = InParameters.GeneralParameters.GroupsToExclude;
 }
 
-const TUniquePtr<ISourceControlProvider>& FSubmitToolPerforce::GetProvider() const
-{
-	return SCCProvider;
-}
-
 UE::Tasks::TTask<TSharedRef<FSCCommand>> FSubmitToolPerforce::GetStream(const FString& InStream)
 {
 	TArray<FString> Args = { TEXT("-o") };
@@ -415,16 +387,6 @@ void FSubmitToolPerforce::SaveRecentUsers()
 	File = nullptr;
 }
 
-bool FSubmitToolPerforce::Tick(float InDeltaTime)
-{
-	if(SCCProvider.IsValid())
-	{
-		SCCProvider->Tick();
-	}
-
-	return true;
-}
-
 const TArray<TSharedPtr<FUserData>>& FSubmitToolPerforce::GetRecentUsers() const
 {
 	return RecentUsers;
@@ -507,14 +469,12 @@ void FSubmitToolPerforce::GetGroups(const FOnGroupsGet::FDelegate& Callback)
 
 				P4Client.SetProtocol("tag", "");
 
-				FSCCRecordSet ResultValues;
-				FSourceControlResultInfo ResultInfo;
 
 				class FGroupsClientUser : public FSTClientUser
 				{
 					public:
-					FGroupsClientUser(FSCCRecordSet& InRecords, EP4ClientUserFlags InFlags, FSourceControlResultInfo& OutResultInfo, TArray<TSharedPtr<FString>>& OutGroupsArray, const TArray<FString>& InGroupsToExclude)
-					: FSTClientUser(InRecords, InFlags, OutResultInfo), GroupsArray(OutGroupsArray), ExcludedGroups(InGroupsToExclude)
+					FGroupsClientUser(FSCCommand& InCommand, EP4ClientUserFlags InFlags, TArray<TSharedPtr<FString>>& OutGroupsArray, const TArray<FString>& InGroupsToExclude)
+					: FSTClientUser(InCommand, InFlags), GroupsArray(OutGroupsArray), ExcludedGroups(InGroupsToExclude)
 					{}
 
 					virtual void OutputStat(StrDict* VarList) override
@@ -567,7 +527,8 @@ void FSubmitToolPerforce::GetGroups(const FOnGroupsGet::FDelegate& Callback)
 
 				TArray<TSharedPtr<FString>> GroupsResults;
 
-				FGroupsClientUser P4User(ResultValues, bIsUnicodeServer ? EP4ClientUserFlags::UnicodeServer : EP4ClientUserFlags::None, ResultInfo, GroupsResults, GroupsToExclude);
+				FSCCommand Command;
+				FGroupsClientUser P4User(Command, bIsUnicodeServer ? EP4ClientUserFlags::UnicodeServer : EP4ClientUserFlags::None, GroupsResults, GroupsToExclude);
 
 				{
 					TRACE_CPUPROFILER_EVENT_SCOPE(FSubmitToolPerforce::RunP4Command_groups);
@@ -575,14 +536,14 @@ void FSubmitToolPerforce::GetGroups(const FOnGroupsGet::FDelegate& Callback)
 					P4Client.Run("groups", &P4User);
 				}
 
-				for(const FText& Msg : ResultInfo.InfoMessages)
+				for(const FText& Msg : Command.Results.InfoMessages)
 				{
 					UE_LOG(LogSubmitToolP4Debug, Verbose, TEXT("p4 groups: %s"), *Msg.ToString());
 				}
 
-				if(ResultInfo.HasErrors())
+				if(Command.Results.HasErrors())
 				{
-					for(const FText& Error : ResultInfo.ErrorMessages)
+					for(const FText& Error : Command.Results.ErrorMessages)
 					{
 						UE_LOG(LogSubmitToolP4, Error, TEXT("p4 groups: %s"), *Error.ToString());
 					}
@@ -598,7 +559,7 @@ void FSubmitToolPerforce::GetGroups(const FOnGroupsGet::FDelegate& Callback)
 					OnGroupsGetCallbacks.Clear();
 				}, LowLevelTasks::ETaskPriority::Normal, UE::Tasks::EExtendedTaskPriority::GameThreadNormalPri);
 
-				return !ResultInfo.HasErrors();
+				return !Command.Results.HasErrors();
 			});
 	}
 }
@@ -902,6 +863,11 @@ const FAuthTicket& FSubmitToolPerforce::GetAuthTicket()
 	return P4Ticket;
 }
 
+void FSubmitToolPerforce::RequestCancel()
+{
+	Connections.CancelCurrentRequests();
+}
+
 bool FSubmitToolPerforce::IsBusy() const
 {
 	return ActiveCommands != 0;
@@ -932,8 +898,11 @@ bool FSubmitToolPerforce::RunP4Command(FSCCommand& Command, EP4ClientUserFlags C
 	ActiveCommands++;
 
 	ClientApi& P4Client = Connection->GetConnection();
-	FString FullCommand = TEXT("p4 -p ");
-	FullCommand += TO_TCHAR(P4Client.GetPort().Text(), bIsUnicodeServer);
+	TStringBuilder<128> ShortCommand;
+	TStringBuilder<512> FullCommand;
+	ShortCommand.Append(TEXT("p4 "));
+	FullCommand.Append(TEXT("p4 -p "));
+	FullCommand.Append(TO_TCHAR(P4Client.GetPort().Text(), bIsUnicodeServer));
 
 	if(EnumHasAnyFlags(CmdFlags, EP4ClientUserFlags::UseUser))
 	{
@@ -941,7 +910,7 @@ bool FSubmitToolPerforce::RunP4Command(FSCCommand& Command, EP4ClientUserFlags C
 		FCmdLineParameters::Get().GetValue(FSubmitToolCmdLine::P4User, UserName);
 		P4Client.SetUser(TCHAR_TO_ANSI(*UserName));
 
-		FullCommand += TEXT(" -u ") + UserName;
+		FullCommand.Append(TEXT(" -u ") + UserName);
 	}
 
 	if(EnumHasAnyFlags(CmdFlags, EP4ClientUserFlags::UseClient))
@@ -950,17 +919,17 @@ bool FSubmitToolPerforce::RunP4Command(FSCCommand& Command, EP4ClientUserFlags C
 		FCmdLineParameters::Get().GetValue(FSubmitToolCmdLine::P4Client, Client);
 		P4Client.SetClient(TCHAR_TO_ANSI(*Client));
 
-		FullCommand += TEXT(" -c ") + Client;
+		FullCommand.Append(TEXT(" -c ") + Client);
 	}
 
 	if(EnumHasAnyFlags(CmdFlags, EP4ClientUserFlags::UseZTag))
 	{
 		P4Client.SetProtocol("tag", "");
-		FullCommand += TEXT(" -ztag");
+		FullCommand.Append(TEXT(" -ztag "));
 	}
 
-	FullCommand += TEXT(" ");
-	FullCommand += Command.Command;
+	ShortCommand.Append(Command.Command);
+	FullCommand.Append(Command.Command);
 
 	int32 ArgC = Command.Args.Num();
 	UTF8CHAR** ArgV = new UTF8CHAR * [ArgC];
@@ -978,18 +947,27 @@ bool FSubmitToolPerforce::RunP4Command(FSCCommand& Command, EP4ClientUserFlags C
 			FMemory::Memcpy(ArgV[Index], TCHAR_TO_ANSI(*Command.Args[Index]), Command.Args[Index].Len() + 1);
 		}
 
-		FullCommand += TEXT(" ") + Command.Args[Index];
+		ShortCommand.Append(TEXT(" "));
+		ShortCommand.Append(Command.Args[Index]);
+		FullCommand.Append(TEXT(" "));
+		FullCommand.Append(Command.Args[Index]);
 	}
 
 	P4Client.SetArgv(ArgC, (char**)ArgV);
-	FSTClientUser P4User(Command.Values, CmdFlags, Command.Results);
+	FSTClientUser P4User(Command, CmdFlags);
 
 	UE_LOG(LogSubmitToolP4Debug, Log, TEXT("Running command: %s"), *FullCommand);
 	P4Client.Run(FROM_TCHAR(*Command.Command, bIsUnicodeServer), &P4User);
 
+	if (P4Client.Dropped())
+	{
+		UE_LOG(LogSubmitToolP4Debug, Log, TEXT("P4 connection dropped: %s"), *Command.Command);
+		Command.bDroppedConnection = true;
+	}
+
 	for(const FText& Msg : Command.Results.InfoMessages)
 	{
-		UE_LOG(LogSubmitToolP4Debug, Verbose, TEXT("%s: %s"), *FullCommand, *Msg.ToString());
+		UE_LOG(LogSubmitToolP4Debug, Log, TEXT("%s: %s"), *FullCommand, *Msg.ToString());
 	}
 
 	if(Command.Results.HasErrors())
@@ -1002,7 +980,8 @@ bool FSubmitToolPerforce::RunP4Command(FSCCommand& Command, EP4ClientUserFlags C
 			}
 			else
 			{
-				UE_LOG(LogSubmitToolP4, Error, TEXT("%s: %s"), *FullCommand, *Error.ToString());
+				UE_LOG(LogSubmitToolP4Debug, Error, TEXT("%s: %s"), *FullCommand, *Error.ToString());
+				UE_LOG(LogSubmitToolP4, Error, TEXT("%s: %s"), *ShortCommand, *Error.ToString());
 			}
 		}
 	}
@@ -1018,5 +997,5 @@ bool FSubmitToolPerforce::RunP4Command(FSCCommand& Command, EP4ClientUserFlags C
 
 
 	ActiveCommands--;
-	return !Command.Results.HasErrors();
+	return !Command.Results.HasErrors() && !Command.bDroppedConnection;
 }
diff --git a/Engine/Source/Programs/SubmitTool/Private/Logic/Services/SourceControl/SubmitToolPerforce.h b/Engine/Source/Programs/SubmitTool/Private/Logic/Services/SourceControl/SubmitToolPerforce.h
--- a/Engine/Source/Programs/SubmitTool/Private/Logic/Services/SourceControl/SubmitToolPerforce.h
+++ b/Engine/Source/Programs/SubmitTool/Private/Logic/Services/SourceControl/SubmitToolPerforce.h
@@ -12,17 +12,22 @@ class ISourceControlProvider;
 
 DECLARE_DELEGATE_RetVal(bool, FOnIsCancelled);
 
-class FP4KeepAlive : public KeepAlive
+class FSTP4KeepAlive : public KeepAlive
 {
 public:
-	FP4KeepAlive(FOnIsCancelled InIsCancelled)
+	FSTP4KeepAlive(FOnIsCancelled InIsCancelled)
 		: IsCancelled(InIsCancelled)
 	{}
 
 	/** Called when the perforce connection wants to know if it should stay connected */
 	virtual int IsAlive() override
 	{
-		return 0;
+		if (IsCancelled.IsBound() && IsCancelled.Execute())
+		{
+			return 0;
+		}
+
+		return 1;
 	}
 
 	FOnIsCancelled IsCancelled;
@@ -30,15 +35,16 @@ public:
 
 struct FClientApiWrapper
 {
-	FClientApiWrapper(TUniquePtr<ClientApi>&& InConnection, FOnIsCancelled InKeepAlive)
-	: bIsReady(true), KeepAlive(InKeepAlive), Connection(Forward<TUniquePtr<ClientApi>>(InConnection))
+	FClientApiWrapper(TUniquePtr<ClientApi>&& InConnection)
+	: bIsReady(true), KeepAlive(FOnIsCancelled::CreateLambda([this] { return bWantsCancel.load(); })), Connection(MoveTemp(InConnection))
 	{
 		Connection->SetBreak(&KeepAlive);
 	}
 
 	std::atomic<bool> bIsReady = true;
-	FP4KeepAlive KeepAlive;
+	FSTP4KeepAlive KeepAlive;
 	TUniquePtr<ClientApi> Connection;
+	std::atomic<bool> bWantsCancel = false;
 };
 
 struct FP4Connection
@@ -58,12 +64,10 @@ class FConnectionPool
 public:
 	FConnectionPool() = default;
 	~FConnectionPool();
-	TUniquePtr<FP4Connection> GetAvailableConnection();
-	
-	void RequestCancel(){ bWantsCancel = true; }
+	TUniquePtr<FP4Connection> GetAvailableConnection();	
+	void CancelCurrentRequests();
 private:
-	size_t CreateConnection();
-	std::atomic<bool> bWantsCancel = false;
+	int32 CreateConnection();
 	std::atomic<bool> bConnectionFailed = false;
 	TArray<TUniquePtr<FClientApiWrapper>> P4Connections;
 	FCriticalSection Mutex;
@@ -77,20 +81,12 @@ public:
 
 
 	virtual void InitializeParameters(const FSubmitToolParameters& InParameters);
-	virtual const TUniquePtr<ISourceControlProvider>& GetProvider() const override;
 	virtual void GetUsers(const FOnUsersGet::FDelegate& Callback) override;
 	virtual void GetGroups(const FOnGroupsGet::FDelegate& Callback) override;
 	virtual void GetUsersAndGroups(const FOnUsersAndGroupsGet::FDelegate& Callback) override;
 
 	virtual UE::Tasks::TTask<bool> DownloadFiles(TArray<FString>&& FilePaths, TMap<FString, FSharedBuffer>& OutFileBuffers) override;
 
-	virtual bool IsAvailable() const override
-	{
-		return SCCProvider.IsValid() && SCCProvider->IsAvailable();
-	}
-
-	virtual bool Tick(float InDeltaTime) override;
-
 	virtual const TArray<TSharedPtr<FUserData>>& GetRecentUsers() const override;
 	virtual void AddRecentUser(TSharedPtr<FUserData>& User) override;
 	virtual const TArray<TSharedPtr<FString>>& GetRecentGroups() const override;
@@ -104,6 +100,7 @@ public:
 	virtual const FString GetCurrentStreamName() override;
 	virtual int32 GetDepotStreamDepth(const FString& InDepotName) override;
 	virtual const FAuthTicket& GetAuthTicket() override;
+	virtual void RequestCancel() override;
 	virtual bool IsBusy() const override;
 			
 protected:
@@ -132,7 +129,6 @@ private:
 	std::atomic<uint64> ActiveCommands = 0;
 	
 	FTSTicker::FDelegateHandle TickHandle;
-	TUniquePtr<ISourceControlProvider> SCCProvider;
 
 	void LoadRecentUsers();
 	void SaveRecentUsers();
diff --git a/Engine/Source/Programs/SubmitTool/Private/Logic/SwarmService.cpp b/Engine/Source/Programs/SubmitTool/Private/Logic/SwarmService.cpp
--- a/Engine/Source/Programs/SubmitTool/Private/Logic/SwarmService.cpp
+++ b/Engine/Source/Programs/SubmitTool/Private/Logic/SwarmService.cpp
@@ -2,9 +2,6 @@
 
 #include "SwarmService.h"
 
-#include "ISourceControlModule.h"
-#include "ISourceControlProvider.h"
-#include "SourceControlOperations.h"
 #include "Interfaces/IHttpResponse.h"
 #include "Logic/Services/SubmitToolServiceProvider.h"
 #include "Logic/Services/Interfaces/ISTSourceControlService.h"
diff --git a/Engine/Source/Programs/SubmitTool/Private/Logic/UpdateService.cpp b/Engine/Source/Programs/SubmitTool/Private/Logic/UpdateService.cpp
--- a/Engine/Source/Programs/SubmitTool/Private/Logic/UpdateService.cpp
+++ b/Engine/Source/Programs/SubmitTool/Private/Logic/UpdateService.cpp
@@ -444,9 +444,6 @@ FString FUpdateService::GetSubmitToolArgs() const
 	FString PerforceClientName;
 	FCmdLineParameters::Get().GetValue(FSubmitToolCmdLine::P4Client, PerforceClientName);
 
-	FString ParameterFile;
-	FCmdLineParameters::Get().GetValue(FSubmitToolCmdLine::ParameterFile, ParameterFile);
-
 	FString FromEditor;
 	FCmdLineParameters::Get().GetValue(FSubmitToolCmdLine::EditorFlag, FromEditor);
 
@@ -455,7 +452,6 @@ FString FUpdateService::GetSubmitToolArgs() const
 	ExecutableArguments = ExecutableArguments.Replace(TEXT("$(client)"), *PerforceClientName);
 	ExecutableArguments = ExecutableArguments.Replace(TEXT("$(user)"), *PerforceUserName);
 	ExecutableArguments = ExecutableArguments.Replace(TEXT("$(changelist)"), *Changelist);
-	ExecutableArguments = ExecutableArguments.Replace(TEXT("$(parameterfile)"), *ParameterFile);
 	ExecutableArguments = ExecutableArguments.Replace(TEXT("$(fromEditor)"), *FromEditor);
 
 	return ExecutableArguments;
diff --git a/Engine/Source/Programs/SubmitTool/Private/Logic/Validators/CoreRedirectsValidator.h b/Engine/Source/Programs/SubmitTool/Private/Logic/Validators/CoreRedirectsValidator.h
--- a/Engine/Source/Programs/SubmitTool/Private/Logic/Validators/CoreRedirectsValidator.h
+++ b/Engine/Source/Programs/SubmitTool/Private/Logic/Validators/CoreRedirectsValidator.h
@@ -4,8 +4,6 @@
 
 #include "ValidatorBaseAsync.h"
 #include "Tasks/Task.h"
-#include "ISourceControlChangelistState.h"
-#include "ISourceControlProvider.h"
 
 class FCoreRedirectsValidator : public FValidatorBaseAsync
 {
diff --git a/Engine/Source/Programs/SubmitTool/Private/Logic/Validators/CrossChangelistValidator.h b/Engine/Source/Programs/SubmitTool/Private/Logic/Validators/CrossChangelistValidator.h
--- a/Engine/Source/Programs/SubmitTool/Private/Logic/Validators/CrossChangelistValidator.h
+++ b/Engine/Source/Programs/SubmitTool/Private/Logic/Validators/CrossChangelistValidator.h
@@ -3,7 +3,6 @@
 #pragma once
 
 #include "ValidatorBase.h"
-#include "ISourceControlChangelistState.h"
 
 class FCrossChangelistValidator : public FValidatorBase
 {
diff --git a/Engine/Source/Programs/SubmitTool/Private/Logic/Validators/PackageDataValidator.h b/Engine/Source/Programs/SubmitTool/Private/Logic/Validators/PackageDataValidator.h
--- a/Engine/Source/Programs/SubmitTool/Private/Logic/Validators/PackageDataValidator.h
+++ b/Engine/Source/Programs/SubmitTool/Private/Logic/Validators/PackageDataValidator.h
@@ -4,8 +4,6 @@
 
 #include "ValidatorBaseAsync.h"
 
-#include "ISourceControlChangelistState.h"
-
 class FPackageDataValidator final : public FValidatorBaseAsync
 {
 public:
diff --git a/Engine/Source/Programs/SubmitTool/Private/Logic/Validators/TagValidator.cpp b/Engine/Source/Programs/SubmitTool/Private/Logic/Validators/TagValidator.cpp
--- a/Engine/Source/Programs/SubmitTool/Private/Logic/Validators/TagValidator.cpp
+++ b/Engine/Source/Programs/SubmitTool/Private/Logic/Validators/TagValidator.cpp
@@ -70,7 +70,7 @@ bool FTagValidator::ValidateTag(const FTag* InTag, const TArray<TSharedPtr<FUser
 {
 	bool bIsTagValid = true;
 
-	FTagValidationConfig ValidationConfig = InTag->GetCurrentValidationConfig(ServiceProvider.Pin()->GetService<FChangelistService>()->GetOpenedFilesDepotPaths());
+	FTagValidationConfig ValidationConfig = InTag->GetCurrentValidationConfig(ServiceProvider.Pin()->GetService<FChangelistService>()->GetFilesInCL());
 
 	if(ValidationConfig.bIsMandatory && !InTag->IsEnabled())
 	{
diff --git a/Engine/Source/Programs/SubmitTool/Private/Logic/Validators/ValidatorBase.cpp b/Engine/Source/Programs/SubmitTool/Private/Logic/Validators/ValidatorBase.cpp
--- a/Engine/Source/Programs/SubmitTool/Private/Logic/Validators/ValidatorBase.cpp
+++ b/Engine/Source/Programs/SubmitTool/Private/Logic/Validators/ValidatorBase.cpp
@@ -77,14 +77,14 @@ void FValidatorBase::StartValidation()
 
 	bool bIncrementalValidation = Definition->bUsesIncrementalCache && !bForceRun;
 	bForceRun = false;
-	TArray<FSourceControlStateRef> IncrementallySkippedFiles;
+	TArray<FSCFileRef> IncrementallySkippedFiles;
 	bool bAppliesToCL = AppliesToCL(ChangelistService->GetCLDescription(), ChangelistService->GetFilesInCL(), TagService->GetTagsArray(), FilteredFiles, IncrementallySkippedFiles, bIncrementalValidation);
 
 	if (bAppliesToCL)
 	{
 		if (FilteredFiles.Num() > 0)
 		{
-			UE_LOG(LogSubmitToolDebug, Log, TEXT("[%s] Validating files:\n - %s"), *ValidatorName, *FString::JoinBy(FilteredFiles, TEXT("\n - "), [](const FSourceControlStateRef& FileRef) { return FileRef->GetFilename(); }));
+			UE_LOG(LogSubmitToolDebug, Log, TEXT("[%s] Validating files:\n - %s"), *ValidatorName, *FString::JoinBy(FilteredFiles, TEXT("\n - "), [](const FSCFileRef& FileRef) { return FileRef->GetFilename(); }));
 		}
 
 		if (!bIsValidSetup)
@@ -173,7 +173,7 @@ void FValidatorBase::StartValidation()
 
 			if (IncrementallySkippedFiles.Num() != 0)
 			{
-				const FString FileList = FString::JoinBy(IncrementallySkippedFiles, TEXT("\n"), [](const FSourceControlStateRef& InFile) { return InFile->GetFilename(); });
+				const FString FileList = FString::JoinBy(IncrementallySkippedFiles, TEXT("\n"), [](const FSCFileRef& InFile) { return InFile->GetFilename(); });
 				UE_LOG(LogValidators, Log, TEXT("[%s] Skipping Files because they were already validated in a previous execution:\n%s"), *GetValidatorName(), *FileList);
 			}
 
@@ -447,11 +447,11 @@ bool FValidatorBase::EvaluateTagSkip()
 
 	if (Definition->SkipForbiddenStreams.Num() > 0)
 	{
-		for (const FString& FileDepotPath : ChangelistService->GetOpenedFilesDepotPaths())
+		for (const FSCFileRef& File : ChangelistService->GetFilesInCL())
 		{
 			for (const FString& Stream : Definition->SkipForbiddenStreams)
 			{
-				if (FileDepotPath.StartsWith(Stream, ESearchCase::IgnoreCase))
+				if (File->GetDepotPath().StartsWith(Stream, ESearchCase::IgnoreCase))
 				{
 					UE_LOG(LogValidators, Log, TEXT("[%s] Changes in %s require the validation to be run from submit tool. %s won't be skipped"), *ValidatorName, *Stream, *ValidatorName);
 					UE_LOG(LogValidatorsResult, Log, TEXT("[%s] Changes in %s require the validation to be run from submit tool. %s won't be skipped"), *ValidatorName, *Stream, *ValidatorName);
diff --git a/Engine/Source/Programs/SubmitTool/Private/Logic/Validators/ValidatorBase.h b/Engine/Source/Programs/SubmitTool/Private/Logic/Validators/ValidatorBase.h
--- a/Engine/Source/Programs/SubmitTool/Private/Logic/Validators/ValidatorBase.h
+++ b/Engine/Source/Programs/SubmitTool/Private/Logic/Validators/ValidatorBase.h
@@ -11,7 +11,6 @@
 
 #include "Models/SCFile.h"
 #include "Logic/Services/SubmitToolServiceProvider.h"
-#include "ISourceControlState.h"
 #include "Logic/Validators/ValidatorDefinition.h"
 #include "ValidatorOptionsProvider.h"
 
diff --git a/Engine/Source/Programs/SubmitTool/Private/Models/ModelInterface.cpp b/Engine/Source/Programs/SubmitTool/Private/Models/ModelInterface.cpp
--- a/Engine/Source/Programs/SubmitTool/Private/Models/ModelInterface.cpp
+++ b/Engine/Source/Programs/SubmitTool/Private/Models/ModelInterface.cpp
@@ -1,6 +1,7 @@
 // Copyright Epic Games, Inc. All Rights Reserved.
 
 #include "ModelInterface.h"
+#include "Algo/Compare.h"
 #include "Models/SubmitToolUserPrefs.h"
 #include "Logic/JiraService.h"
 #include "Logic/PreflightService.h"
@@ -12,7 +13,6 @@
 #include "Widgets/Text/STextBlock.h"
 #include "Containers/Ticker.h"
 #include "Configuration/Configuration.h"
-#include "SourceControlOperations.h"
 #include "HAL/PlatformApplicationMisc.h"
 #include "Async/Async.h"
 #include "Misc/StringOutputDevice.h"
@@ -38,17 +38,14 @@
 ESubmitToolAppState FModelInterface::SubmitToolState = ESubmitToolAppState::Initializing;
 FOnStateChanged FModelInterface::OnStateChanged;
 
-FModelInterface::FModelInterface(const FSubmitToolParameters& InParameters, TSharedPtr<ISTSourceControlService>& InSourceControlService) :
-	Parameters(InParameters), SourceControlService(InSourceControlService)
+FModelInterface::FModelInterface(FSubmitToolParameters&& InParameters, const TSharedPtr<ISTSourceControlService>& InSourceControlService) :
+	Parameters(MoveTemp(InParameters)), SourceControlService(InSourceControlService)
 {
 	TRACE_CPUPROFILER_EVENT_SCOPE(FModelInterface::FModelInterface);
 
 	// initialize call backs
 	CLReadyCallback = FOnChangeListReadyDelegate::CreateRaw(this, &FModelInterface::OnChangelistReady);
 	CLRefreshCallback = FOnChangelistRefreshDelegate::CreateRaw(this, &FModelInterface::OnChangelistRefresh);
-	SubmitFinishedCallback = FSourceControlOperationComplete::CreateRaw(this, &FModelInterface::OnSubmitOperationComplete);
-	DeleteShelveCallback = FSourceControlOperationComplete::CreateRaw(this, &FModelInterface::OnDeleteShelveOperationComplete);
-	RevertUnchangedCallback = FSourceControlOperationComplete::CreateRaw(this, &FModelInterface::OnRevertUnchangedOperationComplete);
 
 	ServiceProvider = MakeShared<FSubmitToolServiceProvider>();
 	// Initialize services
@@ -61,29 +58,29 @@ FModelInterface::FModelInterface(const FSubmitToolParameters& InParameters, TSha
 		ServiceProvider->RegisterService<ICacheDataService>(MakeShared<FCacheDataService>(Parameters.GeneralParameters));
 	}
 
-	SourceControlService->InitializeParameters(InParameters);
+	SourceControlService->InitializeParameters(Parameters);
 	ServiceProvider->RegisterService<ISTSourceControlService>(SourceControlService.ToSharedRef());
-	ValidationService = MakeShared<FTasksService>(InParameters.Validators, TEXT("SubmitTool.StandAlone.Validator"));
+	ValidationService = MakeShared<FTasksService>(Parameters.Validators, TEXT("SubmitTool.StandAlone.Validator"));
 	ServiceProvider->RegisterService<FTasksService>(ValidationService.ToSharedRef(), TEXT("ValidationService"));
-	PresubmitOperationsService = MakeShared<FTasksService>(InParameters.PresubmitOperations, TEXT("SubmitTool.StandAlone.PresubmitOperation"));
+	PresubmitOperationsService = MakeShared<FTasksService>(Parameters.PresubmitOperations, TEXT("SubmitTool.StandAlone.PresubmitOperation"));
 	ServiceProvider->RegisterService<FTasksService>(PresubmitOperationsService.ToSharedRef(), TEXT("PresubmitOperationsService"));
-	CredentialsService = MakeShared<FCredentialsService>(InParameters.OAuthParameters);
+	CredentialsService = MakeShared<FCredentialsService>(Parameters.OAuthParameters);
 	ServiceProvider->RegisterService<FCredentialsService>(CredentialsService.ToSharedRef());
-	ChangelistService = MakeShared<FChangelistService>(InParameters.GeneralParameters, ServiceProvider, CLReadyCallback, CLRefreshCallback);
+	ChangelistService = MakeShared<FChangelistService>(Parameters.GeneralParameters, ServiceProvider, CLReadyCallback, CLRefreshCallback);
 	ServiceProvider->RegisterService<FChangelistService>(ChangelistService.ToSharedRef());
-	P4LockdownService = MakeShared<FP4LockdownService>(InParameters.P4LockdownParameters, ServiceProvider);
+	P4LockdownService = MakeShared<FP4LockdownService>(Parameters.P4LockdownParameters, ServiceProvider);
 	ServiceProvider->RegisterService<FP4LockdownService>(P4LockdownService.ToSharedRef());
-	TagService = MakeShared<FTagService>(InParameters, ChangelistService);
+	TagService = MakeShared<FTagService>(Parameters, ChangelistService);
 	ServiceProvider->RegisterService<FTagService>(TagService.ToSharedRef());
 	SwarmService = MakeShared<FSwarmService>(ServiceProvider);
 	ServiceProvider->RegisterService<FSwarmService>(SwarmService.ToSharedRef());
-	PreflightService = MakeShared<FPreflightService>(InParameters.HordeParameters, this, ServiceProvider);
+	PreflightService = MakeShared<FPreflightService>(Parameters.HordeParameters, this, ServiceProvider);
 	ServiceProvider->RegisterService<FPreflightService>(PreflightService.ToSharedRef());
-	JiraService = MakeShared<FJiraService>(InParameters.JiraParameters, 256, ServiceProvider);
+	JiraService = MakeShared<FJiraService>(Parameters.JiraParameters, 256, ServiceProvider);
 	ServiceProvider->RegisterService<FJiraService>(JiraService.ToSharedRef());
-	IntegrationService = MakeShared<FIntegrationService>(InParameters.IntegrationParameters, ServiceProvider);
+	IntegrationService = MakeShared<FIntegrationService>(Parameters.IntegrationParameters, ServiceProvider);
 	ServiceProvider->RegisterService<FIntegrationService>(IntegrationService.ToSharedRef());
-	UpdateService = MakeShared<FUpdateService>(InParameters.HordeParameters, InParameters.AutoUpdateParameters, ServiceProvider);
+	UpdateService = MakeShared<FUpdateService>(Parameters.HordeParameters, Parameters.AutoUpdateParameters, ServiceProvider);
 	ServiceProvider->RegisterService<FUpdateService>(UpdateService.ToSharedRef());
 
 	ParseValidators();
@@ -180,7 +177,7 @@ void FModelInterface::Dispose() const
 
 	if (GetInputEnabled())
 	{
-		ChangelistService->SendCLDescriptionToP4(EConcurrency::Synchronous);
+		ChangelistService->UpdateP4CLDescriptionSynchronously();
 	}
 
 	for (const TPair<FString, TSharedPtr<FIntegrationOptionBase>>& IntegrationOption : ServiceProvider->GetService<FIntegrationService>()->GetIntegrationOptions())
@@ -342,7 +339,7 @@ void FModelInterface::SendDescriptionToP4() const
 			return;
 		}
 
-		ChangelistService->SendCLDescriptionToP4();
+		ChangelistService->UpdateP4CLDescription();
 	}
 }
 
@@ -540,9 +537,9 @@ void FModelInterface::RequestSwarmReview(const TArray<FString>& InReviewers)
 	{
 		if (!HasShelvedFiles())
 		{
-			ChangelistService->CreateShelvedFiles(FSourceControlOperationComplete::CreateLambda([this, InReviewers](const FSourceControlOperationRef& DeleteShelvedOp, ECommandResult::Type Result)
+			ChangelistService->CreateShelvedFiles([this, InReviewers](const FSCCommand& InCmd)
 				{
-					if (Result == ECommandResult::Succeeded && HasShelvedFiles())
+					if (InCmd.bSuccess && HasShelvedFiles())
 					{
 						RequestSwarmReview(InReviewers);
 					}
@@ -550,7 +547,7 @@ void FModelInterface::RequestSwarmReview(const TArray<FString>& InReviewers)
 					{
 						UE_LOG(LogSubmitTool, Error, TEXT("Failed to shelve files, Swarm Review request is cancelled"));
 					}
-				}));
+				});
 
 			return;
 		}
@@ -656,18 +653,18 @@ void FModelInterface::StartSubmitProcess(bool bSkipShelfDialog)
 
 		if (Parameters.IncompatibleFilesParams.IncompatibleFileGroups.Num() > 0)
 		{
-			const TArray<FString>& FilesInCL = ChangelistService->GetOpenedFilesDepotPaths();
+			const TArray<FSCFileRef>& FilesInCL = ChangelistService->GetFilesInCL();
 
 			for (const FIncompatibleFilesGroup& FileGroup : Parameters.IncompatibleFilesParams.IncompatibleFileGroups)
 			{
 				TArray<size_t, TInlineAllocator<8>> Indexes;
 
-				for (const FString& File : FilesInCL)
+				for (const FSCFileRef& File : FilesInCL)
 				{
 					for (size_t i = 0; i < FileGroup.FileGroups.Num(); ++i)
 					{
 						const FString ReplacedPath = FileGroup.FileGroups[i].Replace(TEXT("$(StreamRoot)"), *GetCurrentStream());
-						if (File.Contains(ReplacedPath, ESearchCase::IgnoreCase))
+						if (File->GetDepotPath().Contains(ReplacedPath, ESearchCase::IgnoreCase))
 						{
 							if (!Indexes.Contains(i))
 							{
@@ -703,48 +700,91 @@ void FModelInterface::StartSubmitProcess(bool bSkipShelfDialog)
 
 		if (HasShelvedFiles())
 		{
-			const TArray<FString>& ShelvedFiles = ChangelistService->GetShelvedFilesDepotPaths();
-			const TArray<FString>& LocalFiles = ChangelistService->GetOpenedFilesDepotPaths();
-
-			if (ShelvedFiles != LocalFiles)
+			if (!bSkipShelfDialog)
 			{
-				const FText TextTitle = NSLOCTEXT("SourceControl.SubmitWindow", "ShelveConflictTitle", "Shelve - Local conflict");
+				const TArray<FSCFileRef>& ShelvedFiles = ChangelistService->GetShelvedFilesInCL();
+				const TArray<FSCFileRef>& LocalFiles = ChangelistService->GetFilesInCL();
+
+				bool bShelfDiffers;
 
-				const size_t MaxFilesToList = 5;
-				TArray<FString> ShelvedList;
-				TArray<FString> LocalList;
-				for (size_t i = 0; i < MaxFilesToList; ++i)
+				if (ShelvedFiles.Num() != LocalFiles.Num())
 				{
-					if (i < ShelvedFiles.Num())
+					bShelfDiffers = true;
+				}
+				else
+				{
+					// Sort and compare open vs. shelved files
+					Algo::StableSortBy(LocalFiles, [](const FSCFileRef& File) { return File->GetDepotPath(); });
+					Algo::StableSortBy(ShelvedFiles, [](const FSCFileRef& File) { return File->GetDepotPath(); });
+					bShelfDiffers = !Algo::Compare(LocalFiles, ShelvedFiles,
+						[](const FSCFileRef& A, const FSCFileRef& B)
+						{
+							return A->GetDepotPath().Equals(B->GetDepotPath(), ESearchCase::IgnoreCase) && A->GetState() == B->GetState();
+						});
+				}
+
+				if (bShelfDiffers)
+				{
+					const FText TextTitle = NSLOCTEXT("SourceControl.SubmitWindow", "ShelveConflictTitle", "Shelve - Local conflict");
+
+					const size_t MaxFilesToList = 5;
+					TArray<FString> ShelvedList;
+					TArray<FString> LocalList;
+					for (size_t i = 0; i < MaxFilesToList; ++i)
 					{
-						ShelvedList.Add(ShelvedFiles[i]);
+						if (i < ShelvedFiles.Num())
+						{
+							ShelvedList.Add(ShelvedFiles[i]->GetDepotPath());
+						}
+
+						if (i < LocalFiles.Num())
+						{
+							LocalList.Add(LocalFiles[i]->GetDepotPath());
+						}
+					}
+					FString LocalListString = FSubmitToolCoreUtils::StringBuilderJoin<1024>(LocalList, TEXT("\n - "));
+					if (LocalFiles.Num() > MaxFilesToList)
+					{
+						LocalListString = FString::Printf(TEXT("%s\n - And %d other files"), *LocalListString, LocalFiles.Num() - MaxFilesToList);
 					}
 
-					if (i < LocalFiles.Num())
+					FString ShelveListString = FSubmitToolCoreUtils::StringBuilderJoin<1024>(ShelvedList, TEXT("\n - "));
+					if (ShelvedFiles.Num() > MaxFilesToList)
 					{
-						LocalList.Add(LocalFiles[i]);
+						ShelveListString = FString::Printf(TEXT("%s\n - And %d other files"), *ShelveListString, ShelvedFiles.Num() - MaxFilesToList);
 					}
-				}
-				FString LocalListString = FSubmitToolCoreUtils::StringBuilderJoin<1024>(LocalList, TEXT("\n - "));
-				if (LocalFiles.Num() > MaxFilesToList)
-				{
-					LocalListString = FString::Printf(TEXT("%s\n - And %d other files"), *LocalListString, LocalFiles.Num() - MaxFilesToList);
-				}
 
-				FString ShelveListString = FSubmitToolCoreUtils::StringBuilderJoin<1024>(ShelvedList, TEXT("\n - "));
-				if (ShelvedFiles.Num() > MaxFilesToList)
-				{
-					ShelveListString = FString::Printf(TEXT("%s\n - And %d other files"), *ShelveListString, ShelvedFiles.Num() - MaxFilesToList);
+					const FString Description =
+						FString::Printf(TEXT("The shelve filelist does not match the local filelist, due to p4 restrictions submit tool can only submit local content do you want to continue with the submit?\nLocal Files:\n - %s\n\nShelved Files:\n - %s"),
+							*LocalListString,
+							*ShelveListString);
+					if (FDialogFactory::ShowDialog(TextTitle, FText::FromString(Description), TArray<FString>{TEXT("Delete Shelve and Submit"), TEXT("Cancel")}) != EDialogFactoryResult::Confirm)
+					{
+						UE_LOG(LogSubmitTool, Log, TEXT("Submission canceled by user"));
+						return;
+					}
 				}
 			}
 
 			ChangeState(ESubmitToolAppState::Submitting, bIsStreamLocked && bIsUserInAllowlist);
-			DeleteShelvedFiles();
+			ChangelistService->DeleteShelvedFiles(
+				[this](const FSCCommand& InCmd) {
+					check(SubmitToolState == ESubmitToolAppState::Submitting);
+					if (InCmd.bSuccess)
+					{
+						Submit();
+					}
+					else
+					{
+						ChangeState(ESubmitToolAppState::WaitingUserInput);
+					}
+				}
+			);
 		}
 		else
 		{
 			ChangeState(ESubmitToolAppState::Submitting, bIsStreamLocked && bIsUserInAllowlist);
-			RevertUnchangedAndSubmit();
+			Submit();
 		}
 	}
 	else
@@ -785,8 +825,8 @@ void FModelInterface::RefreshStateBasedOnFiles()
 			else
 			{
 				// Copy the shelved file depot paths on the main thread since they are not safe to reference from other threads
-				UE::Tasks::Launch(UE_SOURCE_LOCATION, [this, ShelvedFilesDepotPaths = ChangelistService->GetShelvedFilesDepotPaths()] {
-					FPathLockdownResult LockdownStatus = P4LockdownService->ArePathsInLockdown(ShelvedFilesDepotPaths);					
+				UE::Tasks::Launch(UE_SOURCE_LOCATION, [this, ShelvedFilesDepotPaths = ChangelistService->GetShelvedFilesInCL()] {
+					FPathLockdownResult LockdownStatus = P4LockdownService->ArePathsInLockdown(ShelvedFilesDepotPaths);
 					UE::Tasks::Launch(UE_SOURCE_LOCATION, [this, LockdownStatus = MoveTemp(LockdownStatus)]
 					{
 						bIsStreamLocked = LockdownStatus.bIsLocked;
@@ -813,13 +853,13 @@ void FModelInterface::RefreshStateBasedOnFiles()
 		if (!Streams.IsEmpty())
 		{
 			FString StreamsMsg = FString::JoinBy(Streams, TEXT(" -> "), [](const FSCCStream* InStr) { return InStr->Name; });
-			for (const FString& File : ChangelistService->GetOpenedFilesDepotPaths())
+			for (const FSCFileRef& File : ChangelistService->GetFilesInCL())
 			{
 				bool bMappedToView = false;
 
 				for (const FSCCStream* Str : Streams)
 				{
-					if (File.StartsWith(Str->Name))
+					if (File->GetDepotPath().StartsWith(Str->Name))
 					{
 						bMappedToView = true;
 						break;
@@ -827,7 +867,7 @@ void FModelInterface::RefreshStateBasedOnFiles()
 
 					for (const FString& ImportStream : Str->AdditionalImportPaths)
 					{
-						if (File.StartsWith(ImportStream))
+						if (File->GetDepotPath().StartsWith(ImportStream))
 						{
 							bMappedToView = true;
 							break;
@@ -837,14 +877,14 @@ void FModelInterface::RefreshStateBasedOnFiles()
 
 				if (!bMappedToView)
 				{
-					UE_LOG(LogSubmitTool, Warning, TEXT("File %s is not in the stream that the workspace is set to: %s"), *File, *StreamsMsg);
+					UE_LOG(LogSubmitTool, Warning, TEXT("File %s is not in the stream that the workspace is set to: %s"), *File->GetDepotPath(), *StreamsMsg);
 				}
 			}
 		}
 
 		// Copy the file depot paths on the main thread since they are not safe to reference from other threads
-		UE::Tasks::Launch(UE_SOURCE_LOCATION, [this, FilesDepotPaths = ChangelistService->GetOpenedFilesDepotPaths()] {
-			FPathLockdownResult LockdownStatus = P4LockdownService->ArePathsInLockdown(FilesDepotPaths);
+		UE::Tasks::Launch(UE_SOURCE_LOCATION, [this, FilesInCL = ChangelistService->GetFilesInCL()] {
+			FPathLockdownResult LockdownStatus = P4LockdownService->ArePathsInLockdown(FilesInCL);
 			UE::Tasks::Launch(UE_SOURCE_LOCATION, [this, LockdownStatus = MoveTemp(LockdownStatus)]
 			{
 				bIsStreamLocked = LockdownStatus.bIsLocked;
@@ -886,7 +926,8 @@ void FModelInterface::OnChangelistReady(bool bIsValid)
 
 				if (!UpdateService->CheckForNewVersion())
 				{
-					ValidateByArea(ValidateArea);
+					// Delay auto-queuing of validators to allow the UI to be initialized
+					FTSTicker::GetCoreTicker().AddTicker(FTickerDelegate::CreateLambda([this, ValidateArea](float) { ValidateByArea(ValidateArea); return false; }), 0.1f);
 				}
 			}
 
@@ -929,25 +970,6 @@ void FModelInterface::OnChangelistRefresh(ETaskArea InChangeType)
 	}
 }
 
-
-void FModelInterface::RevertUnchangedAndSubmit()
-{
-	bool bHasEditOrAdd = false;
-	for (const FSCFileRef& file : ChangelistService->GetFilesInCL())
-	{
-		bHasEditOrAdd |= file->IsAdded() || file->IsCheckedOut();
-	}
-
-	if (bHasEditOrAdd)
-	{
-		ChangelistService->RevertUnchangedFilesAsync(RevertUnchangedCallback);
-	}
-	else
-	{
-		Submit();
-	}
-}
-
 void FModelInterface::Submit()
 {
 	if (PresubmitOperationsService->AreTasksPendingQueue())
@@ -977,57 +999,13 @@ void FModelInterface::Submit()
 	AddendumAccumulator(Parameters.GeneralParameters.AutomaticChangelistTags, CLDescription, Addendums);
 	Addendums = FConfiguration::Substitute(Addendums);
 
-	ChangelistService->Submit(Addendums, SubmitFinishedCallback);
+	ChangelistService->Submit(Addendums, 
+		[this](const FSCCommand& InCmd) {
+			OnSubmitOperationComplete(InCmd);
+		}		
+	);
 }
 
-void FModelInterface::OnDeleteShelveOperationComplete(const FSourceControlOperationRef& Operation, ECommandResult::Type Result)
-{
-	if (SubmitToolState == ESubmitToolAppState::Submitting)
-	{
-		if (Result == ECommandResult::Succeeded)
-		{
-			RevertUnchangedAndSubmit();
-		}
-		else if (Result == ECommandResult::Failed)
-		{
-			ChangeState(ESubmitToolAppState::WaitingUserInput);
-		}
-		else if (Result == ECommandResult::Cancelled)
-		{
-			ChangeState(ESubmitToolAppState::WaitingUserInput);
-		}
-	}
-	else
-	{
-		ChangeState(ESubmitToolAppState::WaitingUserInput);
-	}
-}
-
-void FModelInterface::OnRevertUnchangedOperationComplete(const FSourceControlOperationRef& Operation, ECommandResult::Type Result)
-{
-	if (SubmitToolState == ESubmitToolAppState::Submitting)
-	{
-		// Revert Unchanged returns as failed if there were no files to revert, check ErrorMessages to see actual failures
-		if (Result == ECommandResult::Cancelled)
-		{
-			ChangeState(ESubmitToolAppState::WaitingUserInput);
-		}
-		else if (Result == ECommandResult::Succeeded || Operation->GetResultInfo().ErrorMessages.Num() == 0)
-		{
-			Submit();
-		}
-		else if (Result == ECommandResult::Failed)
-		{
-			ChangeState(ESubmitToolAppState::WaitingUserInput);
-		}
-	}
-	else
-	{
-		ChangeState(ESubmitToolAppState::WaitingUserInput);
-	}
-}
-
-
 void FModelInterface::OnPresubmitOperationsComplete(bool bInSuccess)
 {
 	if (bInSuccess)
@@ -1042,10 +1020,21 @@ void FModelInterface::OnPresubmitOperationsComplete(bool bInSuccess)
 }
 
 
-void FModelInterface::OnSubmitOperationComplete(const FSourceControlOperationRef& Operation, ECommandResult::Type Result)
+void FModelInterface::OnSubmitOperationComplete(const FSCCommand& InCmd)
 {
-	if (Result == ECommandResult::Succeeded)
+	if (InCmd.bSuccess)
 	{
+		const FString* SubmittedCLString = nullptr;
+		for (const FSCCRecord& Record : InCmd.Values)
+		{
+			SubmittedCLString = Record.Find(TEXT("submittedChange"));
+			if (SubmittedCLString != nullptr)
+			{
+				UE_LOG(LogSubmitToolP4, Log, TEXT("Submitted CL: %s"), **SubmittedCLString);
+				break;
+			}
+		}
+
 		for (const TWeakPtr<const FValidatorBase>& Task : ValidationService->GetTasks())
 		{
 			if (Task.Pin()->Definition->bIsDisabled)
@@ -1065,13 +1054,11 @@ void FModelInterface::OnSubmitOperationComplete(const FSourceControlOperationRef
 		// We've submitted, or tried to submit and failed so we only let the user close the app
 		ChangeState(ESubmitToolAppState::Finished);
 		if(FSubmitToolUserPrefs::Get()->bOpenJiraOnSubmit)
-		{
-			TSharedRef<FCheckIn> CheckIn = StaticCastSharedRef<FCheckIn>(Operation);
-			const FString CLString = CheckIn->GetSuccessMessage().ToString().Replace(TEXT("Submitted changelist "), TEXT(""), ESearchCase::IgnoreCase).TrimStartAndEnd();
-			if (CLString.IsNumeric())
+		{			
+			if (SubmittedCLString != nullptr && SubmittedCLString->IsNumeric())
 			{
-				FPlatformApplicationMisc::ClipboardCopy(*CLString);
-				UE_LOG(LogSubmitToolP4, Log, TEXT("Submitted CL copied to clipboard: %s"), *CLString);
+				FPlatformApplicationMisc::ClipboardCopy(**SubmittedCLString);
+				UE_LOG(LogSubmitToolP4, Log, TEXT("Submitted CL copied to clipboard: %s"), **SubmittedCLString);
 			}
 
 
@@ -1112,7 +1099,7 @@ bool FModelInterface::Tick(float InDeltaTime)
 			break;
 
 		case ESubmitToolAppState::P4BlockingOperation:
-			if (!SwarmService->IsRequestRunning() && !ChangelistService->IsP4OperationRunning() && !JiraService->IsBlockingRequestRunning())
+			if (!SwarmService->IsRequestRunning() && !ChangelistService->HasP4OperationsRunning() && !JiraService->IsBlockingRequestRunning())
 			{
 				ChangeState(ESubmitToolAppState::WaitingUserInput);
 			}
diff --git a/Engine/Source/Programs/SubmitTool/Private/Models/ModelInterface.h b/Engine/Source/Programs/SubmitTool/Private/Models/ModelInterface.h
--- a/Engine/Source/Programs/SubmitTool/Private/Models/ModelInterface.h
+++ b/Engine/Source/Programs/SubmitTool/Private/Models/ModelInterface.h
@@ -77,7 +77,7 @@ DECLARE_MULTICAST_DELEGATE_TwoParams(FOnStateChanged, const ESubmitToolAppState
 class FModelInterface
 {
 public:
-	FModelInterface(const FSubmitToolParameters& InParameters, TSharedPtr<ISTSourceControlService>& InSourceControlService);
+	FModelInterface(FSubmitToolParameters&& InParameters, const TSharedPtr<ISTSourceControlService>& InSourceControlService);
 	~FModelInterface();
 	void Dispose() const;
 	void ParseValidators() const;
@@ -107,12 +107,11 @@ public:
 	FOnCLDescriptionUpdated& GetCLDescriptionUpdatedDelegate()						{ return ChangelistService->OnCLDescriptionUpdated; }
 	
 	const TArray<FSCFileRef>& GetFilesInCL() const									{ return ChangelistService->GetFilesInCL(); }
-	const TArray<FString>& GetDepotFilesInCL() const								{ return ChangelistService->GetOpenedFilesDepotPaths(); }
 	bool HasShelvedFiles() const													{ return ChangelistService->HasShelvedFiles(); }
 
-	bool IsP4OperationRunning(FName OperationName = FName()) const					{ return SubmitToolState == ESubmitToolAppState::P4BlockingOperation || ChangelistService->IsP4OperationRunning(OperationName); }
-	bool IsBlockingOperationRunning() const											{ return SubmitToolState == ESubmitToolAppState::P4BlockingOperation || SwarmService->IsRequestRunning() || JiraService->IsBlockingRequestRunning() || P4LockdownService->IsBlockingOperationRunning(); }
-	void CancelP4Operations(FName OperationName = FName()) const					{ ChangelistService->CancelP4Operations(OperationName); SwarmService->CancelOperations(); }
+	bool IsP4OperationRunning() const												{ return ChangelistService->HasP4OperationsRunning(); }
+	bool IsBlockingOperationRunning() const											{ return SwarmService->IsRequestRunning() || JiraService->IsBlockingRequestRunning() || P4LockdownService->IsBlockingOperationRunning(); }
+	void CancelP4Operations(FName OperationName = FName()) const					{ ChangelistService->CancelP4Operations(); SwarmService->CancelOperations(); }
 
 	void ValidateChangelist() const;
 	void ValidateSingle(const FName& ValidatorId, bool bForce = true) const;
@@ -162,7 +161,9 @@ public:
 	static bool GetInputEnabled();
 	static void SetErrorState()														{ ChangeState(ESubmitToolAppState::Errored);}
 	static const ESubmitToolAppState GetState()										{ return SubmitToolState; }
-	void DeleteShelvedFiles() const													{ ChangelistService->DeleteShelvedFiles(DeleteShelveCallback); }
+	void DeleteShelvedFiles() const													{ ChangelistService->DeleteShelvedFiles(); }
+	void ReplaceShelvedFiles() const												{ ChangelistService->ReplaceShelvedFiles(); }
+
 
 	void RequestPreflight(bool bForceStart = false);
 	bool IsPreflightRequestInProgress() const										{ return PreflightService->IsRequestInProgress(); }
@@ -234,26 +235,20 @@ public:
 	}
 
 private:	
-	const FSubmitToolParameters& Parameters;
+	const FSubmitToolParameters Parameters;
 	void OnChangelistRefresh(ETaskArea InChangeType);
 	void PreValidate(const TArray<const TWeakPtr<const FValidatorBase>>& InValidators) const;
 	void RefreshStateBasedOnFiles();
 	void OnChangelistReady(bool bIsValid);
-	void RevertUnchangedAndSubmit();
 	void Submit();
-	void OnDeleteShelveOperationComplete(const FSourceControlOperationRef& Operation, ECommandResult::Type Result);
-	void OnRevertUnchangedOperationComplete(const FSourceControlOperationRef& Operation, ECommandResult::Type Result);
 	void OnPresubmitOperationsComplete(bool bInSuccess);
-	void OnSubmitOperationComplete(const FSourceControlOperationRef& Operation, ECommandResult::Type Result);
+	void OnSubmitOperationComplete(const FSCCommand& InCmd);
 	void OnGetUsersFromSwarmCompleted(const TUniquePtr<FSwarmReview>& InReview, const FString& InErrorMessage);
 	void OnSwarmCreateCompleted(bool InResult, const FString& InErrorMessage);
 	bool RemoveTagFromDescription(const FString& InTag, FString& InOutCLDescription) const;
 
 	FOnChangeListReadyDelegate CLReadyCallback;
 	FOnChangelistRefreshDelegate CLRefreshCallback;
-	FSourceControlOperationComplete SubmitFinishedCallback;
-	FSourceControlOperationComplete DeleteShelveCallback;
-	FSourceControlOperationComplete RevertUnchangedCallback;
 	FDelegateHandle OnValidationStateUpdatedHandle;
 	FDelegateHandle OnValidationFinishedHandle;
 	FDelegateHandle OnPresubmitFinishedHandle;
diff --git a/Engine/Source/Programs/SubmitTool/Private/Models/SCFile.h b/Engine/Source/Programs/SubmitTool/Private/Models/SCFile.h
deleted file mode 100644
diff --git a/Engine/Source/Programs/SubmitTool/Private/Parameters/SubmitToolParametersBuilder.cpp b/Engine/Source/Programs/SubmitTool/Private/Parameters/SubmitToolParametersBuilder.cpp
--- a/Engine/Source/Programs/SubmitTool/Private/Parameters/SubmitToolParametersBuilder.cpp
+++ b/Engine/Source/Programs/SubmitTool/Private/Parameters/SubmitToolParametersBuilder.cpp
@@ -13,7 +13,7 @@
 #include "Misc/ConfigContext.h"
 #include "Misc/StringOutputDevice.h"
 
-FSubmitToolParametersBuilder::FSubmitToolParametersBuilder(const FString& InParametersXmlFile)
+FSubmitToolParametersBuilder::FSubmitToolParametersBuilder()
 {
 	ConfigHierarchy.Add({ TEXT("Base"), TEXT("{ENGINE}/Config/Base.ini") });
 	ConfigHierarchy.Add( { TEXT("SubmitToolBase"), TEXT("{PROJECT}/Config/{TYPE}.ini") } );
@@ -77,13 +77,22 @@ FSubmitToolParametersBuilder::FSubmitToolParametersBuilder(const FString& InPara
 	ConfigHierarchy.Add( { TEXT("User"), **UserIni, EConfigLayerFlags::NoExpand } );
 }
 
-FSubmitToolParameters FSubmitToolParametersBuilder::Build()
+FSubmitToolParameters FSubmitToolParametersBuilder::LoadConfigFromFiles()
 {
+	TRACE_CPUPROFILER_EVENT_SCOPE(FSubmitToolParametersBuilder::LoadConfigFromFiles);
 	FConfigContext Context = FConfigContext::ReadIntoGConfig();
 	Context.OverrideLayers = ConfigHierarchy;
 
 	FString IniFilename;
 	Context.Load(TEXT("SubmitTool"), IniFilename);
+
+	// if the config was loaded before, force reload it
+	if (SubmitToolConfig != nullptr)
+	{
+		GConfig->SafeUnloadBranch(*IniFilename);
+		SubmitToolConfig = nullptr;
+	}
+
 	SubmitToolConfig = GConfig->FindConfigFile(IniFilename);
 
 	UE_LOG(LogSubmitTool, Verbose, TEXT("Loading config from the following files:"));
diff --git a/Engine/Source/Programs/SubmitTool/Private/Parameters/SubmitToolParametersBuilder.h b/Engine/Source/Programs/SubmitTool/Private/Parameters/SubmitToolParametersBuilder.h
--- a/Engine/Source/Programs/SubmitTool/Private/Parameters/SubmitToolParametersBuilder.h
+++ b/Engine/Source/Programs/SubmitTool/Private/Parameters/SubmitToolParametersBuilder.h
@@ -9,13 +9,13 @@
 class FSubmitToolParametersBuilder
 {
 public:
-	FSubmitToolParametersBuilder(const FString& InParametersXmlFile);
-	FSubmitToolParameters Build();
+	FSubmitToolParametersBuilder();
+	FSubmitToolParameters LoadConfigFromFiles();
 
 private:
 	TArray<FConfigLayer> ConfigHierarchy;
 	TArray<FString> ProjectNames;
-	FConfigFile* SubmitToolConfig;
+	FConfigFile* SubmitToolConfig = nullptr;
 private:
 	FGeneralParameters BuildGeneralParameters();
 	FJiraParameters BuildJiraParameters(); 
diff --git a/Engine/Source/Programs/SubmitTool/Private/SubmitToolApp.cpp b/Engine/Source/Programs/SubmitTool/Private/SubmitToolApp.cpp
--- a/Engine/Source/Programs/SubmitTool/Private/SubmitToolApp.cpp
+++ b/Engine/Source/Programs/SubmitTool/Private/SubmitToolApp.cpp
@@ -210,10 +210,9 @@ int RunSubmitTool(const TCHAR* CommandLine)
 	// App Scope
 	{
 		UE_LOG(LogSubmitTool, Log, TEXT("Version %s"), *FAppVersion::GetVersion());
-		TSharedPtr<ISTSourceControlService> SourceControlService = MakeShared<FSubmitToolPerforce>();
 		
 		// Sync up the config as early as possible
-		UE::Tasks::TTask<TSharedRef<FSCCommand>> SyncCommand = SyncUpConfigFiles(SourceControlService);
+		UE::Tasks::TTask<TUniquePtr<FModelInterface>> InitializeTask = SyncConfigAndInitializeModelInterface();
 		TUniquePtr<FSubmitToolUserPrefs> UserPrefs = FSubmitToolUserPrefs::Initialize(GetUserPrefsPath());
 
 		// allow running multiple parallel instances of UAT
@@ -229,16 +228,7 @@ int RunSubmitTool(const TCHAR* CommandLine)
 		}
 
 		// Wait for the sync since everything else needs the config.
-		SyncCommand.Wait();
-		SyncCommand = { };
-
-		FString ParameterFile;
-		CmdLineParameters.GetValue(FSubmitToolCmdLine::ParameterFile, ParameterFile);
-		FSubmitToolParametersBuilder ParametersBuilder(ParameterFile);
-		FSubmitToolParameters Parameters = ParametersBuilder.Build();
-
-		// Create a new instance of model Interface so that UI can communicate
-		TUniquePtr<FModelInterface> ModelInterface = MakeUnique<FModelInterface>(Parameters, SourceControlService);	
+		TUniquePtr<FModelInterface> ModelInterface = MoveTemp(InitializeTask.GetResult());
 
 		// Build the slate ui
 		TUniquePtr<FSlateUI> SlateUI = MakeSlateUI(ModelInterface.Get());
@@ -246,7 +236,7 @@ int RunSubmitTool(const TCHAR* CommandLine)
 
 		// record that the application has started
 		// do this after initializing UI to give ISTSourceControlService time to populate GetCurrentStreamName
-		FTelemetryService::Init(Parameters.Telemetry);
+		FTelemetryService::Init(ModelInterface->GetParameters().Telemetry);
 		FTelemetryService::Get()->Start(ModelInterface->GetServiceProvider()->GetService<ISTSourceControlService>()->GetCurrentStreamName());
 
 		// Tick the application until shutdown
@@ -268,48 +258,40 @@ int RunSubmitTool(const TCHAR* CommandLine)
 	return 0;
 }
 
-UE::Tasks::TTask<TSharedRef<FSCCommand>> SyncUpConfigFiles(const TSharedPtr<ISTSourceControlService>& InSourceControlService)
+UE::Tasks::TTask<TUniquePtr<FModelInterface>> SyncConfigAndInitializeModelInterface()
 {
-	TRACE_CPUPROFILER_EVENT_SCOPE(FSubmitToolApp::SyncUpConfigFiles);
-	const FString Root = FConfiguration::Substitute(TEXT("$(root)"));
-	TArray<FString> Files{
-		Root / TEXT("Engine/Restricted/NotForLicensees/Config/Mac/MacSubmitTool.ini"),
-		Root / TEXT("Engine/Restricted/NotForLicensees/Config/Linux/LinuxSubmitTool.ini"),
-		Root / TEXT("Engine/Restricted/NotForLicensees/Config/Windows/WindowsSubmitTool.ini"),
-		Root / TEXT("Engine/Restricted/NotForLicensees/Config/SubmitTool.ini")
-	};
-
-	//Files.Add(Root / TEXT("*/Config/Mac/MacSubmitTool.ini"));
-	//Files.Add(Root / TEXT("*/Config/Linux/LinuxSubmitTool.ini"));
-	//Files.Add(Root / TEXT("*/Config/Windows/WindowsSubmitTool.ini"));
-	//Files.Add(Root / TEXT("*/Config/SubmitTool.ini"));
-
-	IFileManager::Get().IterateDirectory(*Root,
-		[&Files](const TCHAR* FileOrDir, bool bIsDir)
+	TRACE_CPUPROFILER_EVENT_SCOPE(FSubmitToolApp::PreloadModules);
+
+	FModuleManager::LoadModuleChecked<FHttpModule>("HTTP");
+
+	return UE::Tasks::Launch(UE_SOURCE_LOCATION, [] {
+
+		TSharedPtr<ISTSourceControlService> SourceControlService = MakeShared<FSubmitToolPerforce>();
+		FSubmitToolParametersBuilder ParametersBuilder;
+		FSubmitToolParameters Parameters = ParametersBuilder.LoadConfigFromFiles();
+		if(Parameters.GeneralParameters.SyncOnStartupPaths.Num() > 0)
 		{
-			if (bIsDir)
+			TRACE_CPUPROFILER_EVENT_SCOPE(FSubmitToolApp::SyncOnStartupFiles);
+			const FString Root = FConfiguration::Substitute(TEXT("$(root)"));
+
+			TArray<FString> Files;
+			Files.Reserve(Parameters.GeneralParameters.SyncOnStartupPaths.Num());
+			for (const FString& Path : Parameters.GeneralParameters.SyncOnStartupPaths)
 			{
-				FString Dir = { FileOrDir };
-				if (!Dir.Contains(TEXT("Engine"), ESearchCase::IgnoreCase))
-				{
-					const FString Extension = TEXT(".uproject");
-					TArray<FString> UProjects;
-					IFileManager::Get().FindFiles(UProjects, FileOrDir, *Extension);
-
-					if (UProjects.Num() != 0)
-					{
-						Files.Add(Dir / TEXT("Config/Mac/MacSubmitTool.ini"));
-						Files.Add(Dir / TEXT("Config/Linux/LinuxSubmitTool.ini"));
-						Files.Add(Dir / TEXT("Config/Windows/WindowsSubmitTool.ini"));
-						Files.Add(Dir / TEXT("Config/SubmitTool.ini"));
-					}
-				}
+				Files.Add(FConfiguration::SubstituteAndNormalizeFilename(Path));
 			}
-			return true;
+
+			SourceControlService->RunSyncCommand(TEXT("sync"), MoveTemp(Files), true);
+
+			// Reload config to be sure we use latest after sync
+			Parameters = ParametersBuilder.LoadConfigFromFiles();
 		}
-	);
 
-	return InSourceControlService->RunCommand(TEXT("sync"), MoveTemp(Files), true);
+
+		TRACE_CPUPROFILER_EVENT_SCOPE(FSubmitToolApp::InitializeModelInterface);
+		// Create a new instance of model Interface so that UI can communicate
+		return MakeUnique<FModelInterface>(MoveTemp(Parameters), SourceControlService);
+	});
 }
 
 FString GetUserPrefsPath()
diff --git a/Engine/Source/Programs/SubmitTool/Private/View/SubmitToolStyle.cpp b/Engine/Source/Programs/SubmitTool/Private/View/SubmitToolStyle.cpp
--- a/Engine/Source/Programs/SubmitTool/Private/View/SubmitToolStyle.cpp
+++ b/Engine/Source/Programs/SubmitTool/Private/View/SubmitToolStyle.cpp
@@ -148,6 +148,15 @@ FSubmitToolStyle::FSubmitToolStyle() : FSlateStyleSet(StyleName)
 	Set("AppIcon.OpenExternal", new IMAGE_BRUSH("Icons/Help/icon_Help_support_16x", FVector2D(12, 12)));
 	Set("AppIcon.Star16", new IMAGE_BRUSH("Icons/Star_16x", FVector2D(16, 16)));
 	
+	// We use a custom teal color for the branched icons
+	BranchedColor = FLinearColor::FromSRGBColor(FColor::FromHex("#00E4A0"));
+	StatusCheckedOutColor = FLinearColor::FromSRGBColor(FColor::FromHex("#1FE44B"));
+
+	Set("RevisionControl.CheckedOut", new CORE_IMAGE_BRUSH_SVG("Starship/SourceControl/SCC_CheckedOut", CoreStyleConstants::Icon16x16, StatusCheckedOutColor));
+	Set("RevisionControl.OpenForAdd", new CORE_IMAGE_BRUSH_SVG("Starship/SourceControl/RC_MarkedForAdd", CoreStyleConstants::Icon16x16, FStyleColors::AccentBlue));
+	Set("RevisionControl.MarkedForDelete", new CORE_IMAGE_BRUSH_SVG("Starship/SourceControl/SCC_MarkedForDelete", CoreStyleConstants::Icon16x16, FStyleColors::Error));
+	Set("RevisionControl.Branched", new CORE_IMAGE_BRUSH_SVG("Starship/SourceControl/SCC_Branched", CoreStyleConstants::Icon16x16, BranchedColor));
+
 	FSlateStyleRegistry::RegisterSlateStyle(*this);
 }
 
diff --git a/Engine/Source/Programs/SubmitTool/Private/View/SubmitToolStyle.h b/Engine/Source/Programs/SubmitTool/Private/View/SubmitToolStyle.h
--- a/Engine/Source/Programs/SubmitTool/Private/View/SubmitToolStyle.h
+++ b/Engine/Source/Programs/SubmitTool/Private/View/SubmitToolStyle.h
@@ -18,6 +18,11 @@ public:
 private:
 	FSubmitToolStyle();
 
+	/** The specific color we use for all the "Branched" icons */
+	FSlateColor BranchedColor;
+	/** The specific colors we use for all the "Status" icons */
+	FSlateColor StatusCheckedOutColor;
+
 	static FName StyleName;
 	static TUniquePtr<FSubmitToolStyle> Inst;
 };
diff --git a/Engine/Source/Programs/SubmitTool/Private/View/SubmitToolWindow.cpp b/Engine/Source/Programs/SubmitTool/Private/View/SubmitToolWindow.cpp
--- a/Engine/Source/Programs/SubmitTool/Private/View/SubmitToolWindow.cpp
+++ b/Engine/Source/Programs/SubmitTool/Private/View/SubmitToolWindow.cpp
@@ -72,7 +72,7 @@ void SubmitToolWindow::CreateMainSubmitToolContent(TSharedPtr<SWindow> InParentW
 
 bool SubmitToolWindow::OnCanCloseTab()
 {
-	if (ModelInterface->IsP4OperationRunning())
+	if (ModelInterface->IsBlockingOperationRunning() || ModelInterface->IsP4OperationRunning())
 	{
 		TSharedPtr<SWindow> WaitForP4Dialog = SNew(SWindow)
 			.Title(FText::FromString("Closing requested"))
@@ -110,15 +110,15 @@ bool SubmitToolWindow::OnCanCloseTab()
 		WaitForP4Dialog->ShowWindow();
 
 		FTSTicker::GetCoreTicker().AddTicker(
-			FTickerDelegate::CreateLambda([Dialog = WaitForP4Dialog, Model = ModelInterface, Tab = MainTab](float InDeltaTime)
+			FTickerDelegate::CreateLambda([WaitForP4Dialog, this](float InDeltaTime)
 			{
-				if (Model->IsP4OperationRunning())
+				if (ModelInterface->IsBlockingOperationRunning() || ModelInterface->IsP4OperationRunning())
 				{
 					return true;
 				}
 
-				Dialog->RequestDestroyWindow();
-				Tab->RequestCloseTab();
+				WaitForP4Dialog->RequestDestroyWindow();
+				MainTab->RequestCloseTab();
 
 				return false;
 			})
diff --git a/Engine/Source/Programs/SubmitTool/Private/View/Widgets/SubmitToolWidget.cpp b/Engine/Source/Programs/SubmitTool/Private/View/Widgets/SubmitToolWidget.cpp
--- a/Engine/Source/Programs/SubmitTool/Private/View/Widgets/SubmitToolWidget.cpp
+++ b/Engine/Source/Programs/SubmitTool/Private/View/Widgets/SubmitToolWidget.cpp
@@ -751,6 +751,25 @@ TSharedRef<SWidget> SubmitToolWidget::BuildOutputLogWidget()
 	return LogTabManager->RestoreFrom(Layout, ParentTab.Pin()->GetParentWindow()).ToSharedRef();
 }
 
+static FSlateIcon GetIcon(ESCFileState State)
+{
+	switch (State)
+	{
+	case ESCFileState::Unknown:
+	case ESCFileState::Edit:
+		return FSlateIcon(FSubmitToolStyle::Get().GetStyleSetName(), "RevisionControl.CheckedOut");
+	case ESCFileState::Add:
+		return FSlateIcon(FSubmitToolStyle::Get().GetStyleSetName(), "RevisionControl.OpenForAdd");
+	case ESCFileState::Delete:
+		return FSlateIcon(FSubmitToolStyle::Get().GetStyleSetName(), "RevisionControl.MarkedForDelete");
+	case ESCFileState::Integrate:
+		return FSlateIcon(FSubmitToolStyle::Get().GetStyleSetName(), "RevisionControl.Branched");
+	default:
+		ensureMsgf(false, TEXT("GetLocalFileState: Unknown perforce file state: '%s'"), *StaticEnum<ESCFileState>()->GetNameStringByValue(static_cast<int64>(State)));
+		return {};
+	}
+}
+
 TSharedRef<SExpandableArea> SubmitToolWidget::BuildFilesInCLWidget()
 {
 	TSharedPtr<SListView<FSCFileRef>> FileList = SNew(SListView<FSCFileRef>)
@@ -768,7 +787,7 @@ TSharedRef<SExpandableArea> SubmitToolWidget::BuildFilesInCLWidget()
 						.HAlign(HAlign_Left)
 						[
 							SNew(SImage)
-							.Image(InItem->GetIcon().GetIcon())
+							.Image(GetIcon(InItem->GetState()).GetIcon())
 						]
 
 						+ SHorizontalBox::Slot()
diff --git a/Engine/Source/Programs/SubmitTool/Private/View/Widgets/TagWidget.cpp b/Engine/Source/Programs/SubmitTool/Private/View/Widgets/TagWidget.cpp
--- a/Engine/Source/Programs/SubmitTool/Private/View/Widgets/TagWidget.cpp
+++ b/Engine/Source/Programs/SubmitTool/Private/View/Widgets/TagWidget.cpp
@@ -88,7 +88,7 @@ void STagWidget::Construct(const FArguments& InArgs)
 			switch(Tag->GetTagState())
 			{			
 				case ETagState::Unchecked:
-					if(Tag->GetCurrentValidationConfig(ModelInterface->GetDepotFilesInCL()).bIsMandatory && !Tag->IsEnabled())
+					if(Tag->GetCurrentValidationConfig(ModelInterface->GetFilesInCL()).bIsMandatory && !Tag->IsEnabled())
 					{
 						return FAppStyle::GetColor("ValidatorStateFail");
 					}
@@ -116,7 +116,7 @@ void STagWidget::Construct(const FArguments& InArgs)
 				default:
 					return FLinearColor::White;
 			}})
-			.Text_Lambda([this](){ return FText::FromString(FString::Printf(TEXT("%s%s"), *this->Tag->Definition.TagLabel, Tag->GetCurrentValidationConfig(ModelInterface->GetDepotFilesInCL()).bIsMandatory ? TEXT(" *") : TEXT(""))); })
+			.Text_Lambda([this](){ return FText::FromString(FString::Printf(TEXT("%s%s"), *this->Tag->Definition.TagLabel, Tag->GetCurrentValidationConfig(ModelInterface->GetFilesInCL()).bIsMandatory ? TEXT(" *") : TEXT(""))); })
 			.ToolTipText(FText::FromString(Tag->Definition.ToolTip))
 			.Font(FAppStyle::Get().GetFontStyle("SmallFont"))
 		]
diff --git a/Engine/Source/Programs/SubmitTool/Public/SubmitToolApp.h b/Engine/Source/Programs/SubmitTool/Public/SubmitToolApp.h
--- a/Engine/Source/Programs/SubmitTool/Public/SubmitToolApp.h
+++ b/Engine/Source/Programs/SubmitTool/Public/SubmitToolApp.h
@@ -6,11 +6,11 @@
 #include "Tasks/Task.h"
 
 class ISTSourceControlService;
-struct FSCCommand;
+class FModelInterface;
 
 /**
  * Run the SubmitTool .
  */
 int RunSubmitTool(const TCHAR* Commandline);
-UE::Tasks::TTask<TSharedRef<FSCCommand>> SyncUpConfigFiles(const TSharedPtr<ISTSourceControlService>& InSourceControlService);
+UE::Tasks::TTask<TUniquePtr<FModelInterface>> SyncConfigAndInitializeModelInterface();
 FString GetUserPrefsPath();
diff --git a/Engine/Source/Programs/SubmitTool/SubmitTool.Build.cs b/Engine/Source/Programs/SubmitTool/SubmitTool.Build.cs
--- a/Engine/Source/Programs/SubmitTool/SubmitTool.Build.cs
+++ b/Engine/Source/Programs/SubmitTool/SubmitTool.Build.cs
@@ -30,7 +30,6 @@ public class SubmitTool : ModuleRules
 				"StandaloneRenderer",
 				"StudioTelemetry",
 				"ToolWidgets",
-				"SourceControl",
 				"DesktopPlatform", // Open file dialog
 				"Settings",
 				"OutputLog",
@@ -44,12 +43,6 @@ public class SubmitTool : ModuleRules
 			}
 		);
 
-		DynamicallyLoadedModuleNames.AddRange(
-			new string[] {
-				"PerforceSourceControl",
-			}
-		);
-
 
 		if (Target.Configuration != UnrealTargetConfiguration.Shipping)
 		{
@@ -109,5 +102,9 @@ public class SubmitTool : ModuleRules
 		RuntimeDependencies.Add("$(EngineDir)/Content/Editor/Slate/Old/DashedBorder.png", StagedFileType.UFS); 
 		RuntimeDependencies.Add("$(EngineDir)/Content/Editor/Slate/Old/ToolTip_Background.png", StagedFileType.UFS);
 		RuntimeDependencies.Add("$(EngineDir)/Content/Editor/Slate/Checkerboard.png", StagedFileType.UFS);
+		RuntimeDependencies.Add("$(EngineDir)/Content/Slate/Starship/SourceControl/SCC_CheckedOut.svg", StagedFileType.UFS);
+		RuntimeDependencies.Add("$(EngineDir)/Content/Slate/Starship/SourceControl/RC_MarkedForAdd.svg", StagedFileType.UFS);
+		RuntimeDependencies.Add("$(EngineDir)/Content/Slate/Starship/SourceControl/SCC_MarkedForDelete.svg", StagedFileType.UFS);
+		RuntimeDependencies.Add("$(EngineDir)/Content/Slate/Starship/SourceControl/SCC_Branched.svg", StagedFileType.UFS);
 	}
 }
diff --git a/Engine/Source/Programs/SubmitTool/SubmitTool.Target.cs b/Engine/Source/Programs/SubmitTool/SubmitTool.Target.cs
--- a/Engine/Source/Programs/SubmitTool/SubmitTool.Target.cs
+++ b/Engine/Source/Programs/SubmitTool/SubmitTool.Target.cs
@@ -35,7 +35,6 @@ public class SubmitToolTarget : TargetRules
 		bForceDisableAutomationTests	= true;
 
 		bHasExports = false;
-		EnablePlugins.Add("PerforceSourceControl");
 
 		WindowsPlatform.bStripUnreferencedSymbols = true;
 		WindowsPlatform.bMergeIdenticalCOMDATs = true;
