Jira migration from search to search/jql

#Synced-CL 49723226

diff --git a/Engine/Source/Programs/SubmitTool/Private/Logic/JiraService.cpp b/Engine/Source/Programs/SubmitTool/Private/Logic/JiraService.cpp
--- a/Engine/Source/Programs/SubmitTool/Private/Logic/JiraService.cpp
+++ b/Engine/Source/Programs/SubmitTool/Private/Logic/JiraService.cpp
@@ -74,7 +74,7 @@ void FJiraService::QueryIssues()
 	JiraRequest->OnProcessRequestComplete().BindRaw(this, &FJiraService::QueryIssues_HttpRequestComplete);
 
 	TSharedPtr<FCredentialsService> CredentialsService = ServiceProvider.Pin()->GetService<FCredentialsService>();
-	FString Url = FString::Format(TEXT("https://{0}/rest/api/2/search?maxResults={1}&jql=assignee={2}"), { Definition.ServerAddress, this->MaxResults, CredentialsService->GetUsername()});
+	FString Url = FString::Format(TEXT("https://{0}/rest/api/2/search/jql?maxResults={1}&jql=assignee={2}"), { Definition.ServerAddress, this->MaxResults, CredentialsService->GetUsername()});
 
 	JiraRequest->SetURL(Url);
 	JiraRequest->SetHeader(TEXT("Authorization"), TEXT("Basic ") + CredentialsService->GetEncodedLoginString());
