Supporting new field in Integration Request.
Changes for jira cloud

#Synced-CL 49723226

diff --git a/Engine/Source/Developer/SubmitToolCore/Public/Parameters/SubmitToolParameters.h b/Engine/Source/Developer/SubmitToolCore/Public/Parameters/SubmitToolParameters.h
--- a/Engine/Source/Developer/SubmitToolCore/Public/Parameters/SubmitToolParameters.h
+++ b/Engine/Source/Developer/SubmitToolCore/Public/Parameters/SubmitToolParameters.h
@@ -96,6 +96,9 @@ struct FJiraParameters
 
 	UPROPERTY()
 	FString PreflightField;
+
+	UPROPERTY()
+	FString JiraField;
 	
 	UPROPERTY()
 	FString StreamField;
diff --git a/Engine/Source/Programs/SubmitTool/Private/Logic/JiraService.cpp b/Engine/Source/Programs/SubmitTool/Private/Logic/JiraService.cpp
--- a/Engine/Source/Programs/SubmitTool/Private/Logic/JiraService.cpp
+++ b/Engine/Source/Programs/SubmitTool/Private/Logic/JiraService.cpp
@@ -479,6 +479,21 @@ void FJiraService::CreateServiceDeskRequest(TSharedPtr<FJsonObject> InBaseJiraJs
 		}
 	}
 
+	if (!Definition.JiraField.IsEmpty())
+	{
+		const FTag* JiraTag = ServiceProvider.Pin()->GetService<FTagService>()->GetTagOfType(TEXT("JiraIssue"));
+		if (JiraTag != nullptr && JiraTag->GetValues().Num() > 0)
+		{
+			TArray<FString> URLs;
+			for (const FString& TagValue : JiraTag->GetValues())
+			{
+				URLs.Add(FString::Format(TEXT("https://{0}/browse/{1}"), { Definition.ServerAddress, TagValue }));
+			}
+
+			RequestFieldValuesJson->SetStringField(Definition.JiraField, FString::Join(URLs, TEXT("\n")));
+		}
+	}
+
 	FString Requestor = FConfiguration::Substitute(TEXT("$(USER)"));
 	TSharedPtr<FUserData> LocalUserData = ServiceProvider.Pin()->GetService<ISTSourceControlService>()->GetUserDataFromCache(Requestor);
 	if (LocalUserData.IsValid())
diff --git a/Engine/Source/Programs/SubmitTool/Private/View/Widgets/SIntegrationWidget.cpp b/Engine/Source/Programs/SubmitTool/Private/View/Widgets/SIntegrationWidget.cpp
--- a/Engine/Source/Programs/SubmitTool/Private/View/Widgets/SIntegrationWidget.cpp
+++ b/Engine/Source/Programs/SubmitTool/Private/View/Widgets/SIntegrationWidget.cpp
@@ -101,7 +101,7 @@ void SIntegrationWidget::Construct(const FArguments& InArgs)
 				.ButtonStyle(FAppStyle::Get(), "PrimaryButton")
 				.IsEnabled_Lambda([this]
 					{
-						return bAreFieldsValid && !ModelInterface->IsBlockingOperationRunning() && ModelInterface->HasSubmitToolTag() && ModelInterface->IsIntegrationRequired();
+						return ModelInterface->GetState() != ESubmitToolAppState::Finished && bAreFieldsValid && !ModelInterface->IsBlockingOperationRunning() && ModelInterface->HasSubmitToolTag() && ModelInterface->IsIntegrationRequired();
 					})
 				.OnClicked(this, &SIntegrationWidget::OnRequestIntegrationClicked)
 			]
