Adding option to disable jira integration
Bumping submit tool version to 6.6.1

#Synced-CL 49723226

diff --git a/Engine/Source/Developer/SubmitToolCore/Public/Parameters/SubmitToolParameters.h b/Engine/Source/Developer/SubmitToolCore/Public/Parameters/SubmitToolParameters.h
--- a/Engine/Source/Developer/SubmitToolCore/Public/Parameters/SubmitToolParameters.h
+++ b/Engine/Source/Developer/SubmitToolCore/Public/Parameters/SubmitToolParameters.h
@@ -76,6 +76,9 @@ struct FJiraParameters
 {
 	GENERATED_BODY()
 
+	UPROPERTY()
+	bool bDisableJira = false;
+
 	UPROPERTY()
 	FString ServerAddress;
 	
diff --git a/Engine/Source/Programs/SubmitTool/Private/Logic/JiraService.cpp b/Engine/Source/Programs/SubmitTool/Private/Logic/JiraService.cpp
--- a/Engine/Source/Programs/SubmitTool/Private/Logic/JiraService.cpp
+++ b/Engine/Source/Programs/SubmitTool/Private/Logic/JiraService.cpp
@@ -44,7 +44,7 @@ FJiraService::~FJiraService()
 bool FJiraService::FetchJiraTickets(bool InForce)
 {
 	TSharedPtr<FCredentialsService> CredentialsService = ServiceProvider.Pin()->GetService<FCredentialsService>();
-	if(CredentialsService->HasCredentials() && !Definition.ServerAddress.IsEmpty())
+	if(CredentialsService->HasCredentials() && !Definition.ServerAddress.IsEmpty() && !Definition.bDisableJira)
 	{
 		if (InForce || (JiraIssues.Num() == 0 && CredentialsService->AreCredentialsValid()))
 		{
@@ -63,7 +63,7 @@ void FJiraService::Reset()
 
 void FJiraService::QueryIssues()
 {
-	if(JiraRequest.IsValid())
+	if(JiraRequest.IsValid() || Definition.bDisableJira)
 	{
 		return;
 	}
diff --git a/Engine/Source/Programs/SubmitTool/Private/Version/AppVersionDefines.h b/Engine/Source/Programs/SubmitTool/Private/Version/AppVersionDefines.h
--- a/Engine/Source/Programs/SubmitTool/Private/Version/AppVersionDefines.h
+++ b/Engine/Source/Programs/SubmitTool/Private/Version/AppVersionDefines.h
@@ -21,7 +21,7 @@
 
 #define SUBMIT_TOOL_VERSION_MAJOR 6
 #define SUBMIT_TOOL_VERSION_MINOR 6
-#define SUBMIT_TOOL_VERSION_PATCH 0
+#define SUBMIT_TOOL_VERSION_PATCH 1
 
 #ifndef SUBMIT_TOOL_CHANGELIST
 #define SUBMIT_TOOL_CHANGELIST_STRING TEXT("NoCL")
diff --git a/Engine/Source/Programs/SubmitTool/Private/View/Widgets/TagWidget.cpp b/Engine/Source/Programs/SubmitTool/Private/View/Widgets/TagWidget.cpp
--- a/Engine/Source/Programs/SubmitTool/Private/View/Widgets/TagWidget.cpp
+++ b/Engine/Source/Programs/SubmitTool/Private/View/Widgets/TagWidget.cpp
@@ -376,7 +376,8 @@ void STagWidget::Construct(const FArguments& InArgs)
 				.MinDesiredWidth(60)
 				[
 					SNew(SButton)
-					.IsEnabled_Static(&FModelInterface::GetInputEnabled)
+					.IsEnabled_Lambda([this] { return ModelInterface->GetInputEnabled() && !ModelInterface->GetParameters().JiraParameters.bDisableJira; })
+					.ToolTipText_Lambda([this] { return ModelInterface->GetParameters().JiraParameters.bDisableJira ? FText::FromString(TEXT("Jira integration is disabled")) : FText::FromString(TEXT("Browse Jira issues")); })
 					.OnClicked(this, &STagWidget::OnJiraClick)
 					.HAlign(HAlign_Center)
 					.VAlign(VAlign_Center)
