Adding priority index for the submit tool loaded config files printing

#Synced-CL 49723226

diff --git a/Engine/Source/Programs/SubmitTool/Private/Parameters/SubmitToolParametersBuilder.cpp b/Engine/Source/Programs/SubmitTool/Private/Parameters/SubmitToolParametersBuilder.cpp
--- a/Engine/Source/Programs/SubmitTool/Private/Parameters/SubmitToolParametersBuilder.cpp
+++ b/Engine/Source/Programs/SubmitTool/Private/Parameters/SubmitToolParametersBuilder.cpp
@@ -101,7 +101,7 @@ FSubmitToolParameters FSubmitToolParametersBuilder::LoadConfigFromFiles()
 		FString FullString(OverrideLayer.Value);
 		if(IFileManager::Get().FileExists(*FullString))
 		{
-			UE_LOG(LogSubmitTool, Verbose, TEXT("%s"), *FPaths::ConvertRelativePathToFull(FullString));
+			UE_LOG(LogSubmitTool, Verbose, TEXT("Priority: %d - File:%s"), OverrideLayer.Key, *FPaths::ConvertRelativePathToFull(FullString));
 		}
 	}
 
