Submit tool config sync+load is now done in a single pass, fixing reloading issues on Mac

#Synced-CL 50047099
#Synced-CL 50047099

diff --git a/Engine/Source/Programs/SubmitTool/Private/Parameters/SubmitToolParametersBuilder.cpp b/Engine/Source/Programs/SubmitTool/Private/Parameters/SubmitToolParametersBuilder.cpp
--- a/Engine/Source/Programs/SubmitTool/Private/Parameters/SubmitToolParametersBuilder.cpp
+++ b/Engine/Source/Programs/SubmitTool/Private/Parameters/SubmitToolParametersBuilder.cpp
@@ -85,15 +85,7 @@ FSubmitToolParameters FSubmitToolParametersBuilder::LoadConfigFromFiles()
 
 	FString IniFilename;
 	Context.Load(TEXT("SubmitTool"), IniFilename);
-
-	// if the config was loaded before, force reload it
-	if (SubmitToolConfig != nullptr)
-	{
-		GConfig->SafeUnloadBranch(*IniFilename);
-		SubmitToolConfig = nullptr;
-	}
-
-	SubmitToolConfig = GConfig->FindConfigFile(IniFilename);
+	FConfigFile* SubmitToolConfig = GConfig->FindConfigFile(IniFilename);
 
 	UE_LOG(LogSubmitTool, Verbose, TEXT("Loading config from the following files:"));
 	for(const TPair<int32, FUtf8String>& OverrideLayer : Context.Branch->Hierarchy)
@@ -106,25 +98,43 @@ FSubmitToolParameters FSubmitToolParametersBuilder::LoadConfigFromFiles()
 	}
 
 	FSubmitToolParameters Parameters;
-	Parameters.GeneralParameters = BuildGeneralParameters();
-	Parameters.JiraParameters = BuildJiraParameters();
-	Parameters.Telemetry = GetTelemetryParameters();
-	Parameters.IntegrationParameters = BuildIntegrationParameters();
-	Parameters.AvailableTags = BuildAvailableTags();
-	Parameters.Validators = BuildValidators();
-	Parameters.PresubmitOperations = BuildPresubmitOperations();
-	Parameters.CopyLogParameters = BuildCopyLogParameters();
-	Parameters.P4LockdownParameters = BuildP4LockdownParameters();
-	Parameters.OAuthParameters = BuildOAuthParameters();
-	Parameters.IncompatibleFilesParams = BuildIncompatibleFilesParameters();
-	Parameters.HordeParameters = BuildHordeParameters();
-	Parameters.AutoUpdateParameters = BuildAutoUpdateParameters();
+
+	if (SubmitToolConfig == nullptr)
+	{
+		UE_LOG(LogSubmitTool, Error, TEXT("Failed to load config file"));
+		return Parameters;
+	}
+
+	Parameters.GeneralParameters = BuildGeneralParameters(*SubmitToolConfig);
+	Parameters.JiraParameters = BuildJiraParameters(*SubmitToolConfig);
+	Parameters.Telemetry = GetTelemetryParameters(*SubmitToolConfig);
+	Parameters.IntegrationParameters = BuildIntegrationParameters(*SubmitToolConfig);
+	Parameters.AvailableTags = BuildAvailableTags(*SubmitToolConfig);
+	Parameters.Validators = BuildValidators(*SubmitToolConfig);
+	Parameters.PresubmitOperations = BuildPresubmitOperations(*SubmitToolConfig);
+	Parameters.CopyLogParameters = BuildCopyLogParameters(*SubmitToolConfig);
+	Parameters.P4LockdownParameters = BuildP4LockdownParameters(*SubmitToolConfig);
+	Parameters.OAuthParameters = BuildOAuthParameters(*SubmitToolConfig);
+	Parameters.IncompatibleFilesParams = BuildIncompatibleFilesParameters(*SubmitToolConfig);
+	Parameters.HordeParameters = BuildHordeParameters(*SubmitToolConfig);
+	Parameters.AutoUpdateParameters = BuildAutoUpdateParameters(*SubmitToolConfig);
 	return Parameters;
 }
 
-FGeneralParameters FSubmitToolParametersBuilder::BuildGeneralParameters()
+FGeneralParameters FSubmitToolParametersBuilder::ReadGeneralParametersFromLocalConfig()
+{
+	FConfigFile LocalConfigFile;
+	FConfigContext Context = FConfigContext::ReadSingleIntoLocalFile(LocalConfigFile);
+
+	FString IniFilename;
+	Context.Load(TEXT("SubmitTool"), IniFilename);
+
+	return BuildGeneralParameters(LocalConfigFile);
+}
+
+FGeneralParameters FSubmitToolParametersBuilder::BuildGeneralParameters(const FConfigFile& InConfigFile)
 {
-	const FConfigSection* Section = SubmitToolConfig->FindSection(TEXT("SubmitTool.General"));
+	const FConfigSection* Section = InConfigFile.FindSection(TEXT("SubmitTool.General"));
 	FGeneralParameters Output;
 
 	if(Section != nullptr)
@@ -159,9 +169,9 @@ FGeneralParameters FSubmitToolParametersBuilder::BuildGeneralParameters()
 	return Output;
 }
 
-FJiraParameters FSubmitToolParametersBuilder::BuildJiraParameters()
+FJiraParameters FSubmitToolParametersBuilder::BuildJiraParameters(const FConfigFile& InConfigFile)
 {
-	const FConfigSection* Section = SubmitToolConfig->FindSection(TEXT("SubmitTool.Jira"));
+	const FConfigSection* Section = InConfigFile.FindSection(TEXT("SubmitTool.Jira"));
 	FJiraParameters Output;
 
 	if(Section != nullptr)
@@ -179,7 +189,7 @@ FJiraParameters FSubmitToolParametersBuilder::BuildJiraParameters()
 	return Output;
 }
 
- FTelemetryParameters FSubmitToolParametersBuilder::GetTelemetryParameters()
+ FTelemetryParameters FSubmitToolParametersBuilder::GetTelemetryParameters(const FConfigFile& InConfigFile)
 {
 	static const TCHAR* Section = TEXT("SubmitTool.Telemetry");
 	static const TCHAR* UrlKey = TEXT("URL");
@@ -188,13 +198,13 @@ FJiraParameters FSubmitToolParametersBuilder::BuildJiraParameters()
 	FTelemetryParameters Output;
 
 	FString Url;
-	if (SubmitToolConfig->GetString(Section, UrlKey, Url))
+	if (InConfigFile.GetString(Section, UrlKey, Url))
 	{
 		Output.Url = Url;
 	}
 
 	FString Instance;
-	if (SubmitToolConfig->GetString(Section, InstanceKey, Instance))
+	if (InConfigFile.GetString(Section, InstanceKey, Instance))
 	{
 		Output.Instance = Instance;
 	}
@@ -202,9 +212,9 @@ FJiraParameters FSubmitToolParametersBuilder::BuildJiraParameters()
 	return Output;
 }
 
- FIntegrationParameters FSubmitToolParametersBuilder::BuildIntegrationParameters()
+ FIntegrationParameters FSubmitToolParametersBuilder::BuildIntegrationParameters(const FConfigFile& InConfigFile)
  {
-	 const FConfigSection* Section = SubmitToolConfig->FindSection(TEXT("SubmitTool.FNIntegration"));
+	 const FConfigSection* Section = InConfigFile.FindSection(TEXT("SubmitTool.FNIntegration"));
 	 FIntegrationParameters Output;
 
 	 if(Section != nullptr)
@@ -222,12 +232,12 @@ FJiraParameters FSubmitToolParametersBuilder::BuildJiraParameters()
 	 return Output;
  }
 
-TArray<FTagDefinition> FSubmitToolParametersBuilder::BuildAvailableTags()
+TArray<FTagDefinition> FSubmitToolParametersBuilder::BuildAvailableTags(const FConfigFile& InConfigFile)
 {
 	static const TCHAR* TagsSectionName = TEXT("Tags.");
 	TArray<FTagDefinition> Output;
 
-	for(const TPair<FString, FConfigSection>& Section : AsConst(*SubmitToolConfig))
+	for(const TPair<FString, FConfigSection>& Section : InConfigFile)
 	{
 		if(Section.Key.StartsWith(TagsSectionName))
 		{
@@ -267,13 +277,13 @@ TArray<FTagDefinition> FSubmitToolParametersBuilder::BuildAvailableTags()
 	return Output;
 }
 
-TMap<FString, FString> FSubmitToolParametersBuilder::BuildValidators()
+TMap<FString, FString> FSubmitToolParametersBuilder::BuildValidators(const FConfigFile& InConfigFile)
 {
 	static const TCHAR* ValidatorsSectionName = TEXT("Validator.");
 
 	TMap<FString, FString> Output;
 
-	for(const TPair<FString, FConfigSection>& Section : AsConst(*SubmitToolConfig))
+	for(const TPair<FString, FConfigSection>& Section : InConfigFile)
 	{
 		if(Section.Key.StartsWith(ValidatorsSectionName))
 		{
@@ -284,13 +294,13 @@ TMap<FString, FString> FSubmitToolParametersBuilder::BuildValidators()
 	return Output;
 }
 
-TMap<FString, FString> FSubmitToolParametersBuilder::BuildPresubmitOperations()
+TMap<FString, FString> FSubmitToolParametersBuilder::BuildPresubmitOperations(const FConfigFile& InConfigFile)
 {
 	static const TCHAR* PresubmitOperationsSectionName = TEXT("PresubmitOperation.");
 
 	TMap<FString, FString> Output;
 
-	for(const TPair<FString, FConfigSection>& Section : AsConst(*SubmitToolConfig))
+	for(const TPair<FString, FConfigSection>& Section : InConfigFile)
 	{
 		if(Section.Key.StartsWith(PresubmitOperationsSectionName))
 		{
@@ -301,9 +311,9 @@ TMap<FString, FString> FSubmitToolParametersBuilder::BuildPresubmitOperations()
 	return Output;
 }
 
-FCopyLogParameters FSubmitToolParametersBuilder::BuildCopyLogParameters()
+FCopyLogParameters FSubmitToolParametersBuilder::BuildCopyLogParameters(const FConfigFile& InConfigFile)
 {
-	const FConfigSection* Section = SubmitToolConfig->FindSection(TEXT("SubmitTool.CopyLog"));
+	const FConfigSection* Section = InConfigFile.FindSection(TEXT("SubmitTool.CopyLog"));
 	FCopyLogParameters Output;
 	
 	if(Section != nullptr)
@@ -320,9 +330,9 @@ FCopyLogParameters FSubmitToolParametersBuilder::BuildCopyLogParameters()
 
 	return Output;
 }
-FP4LockdownParameters FSubmitToolParametersBuilder::BuildP4LockdownParameters()
+FP4LockdownParameters FSubmitToolParametersBuilder::BuildP4LockdownParameters(const FConfigFile& InConfigFile)
 {
-	const FConfigSection* Section = SubmitToolConfig->FindSection(TEXT("SubmitTool.P4Lockdown"));
+	const FConfigSection* Section = InConfigFile.FindSection(TEXT("SubmitTool.P4Lockdown"));
 	FP4LockdownParameters Output;
 
 	if(Section != nullptr)
@@ -340,9 +350,9 @@ FP4LockdownParameters FSubmitToolParametersBuilder::BuildP4LockdownParameters()
 	return Output;
 }
 
-FOAuthTokenParams FSubmitToolParametersBuilder::BuildOAuthParameters()
+FOAuthTokenParams FSubmitToolParametersBuilder::BuildOAuthParameters(const FConfigFile& InConfigFile)
 {
-	const FConfigSection* Section = SubmitToolConfig->FindSection(TEXT("SubmitTool.OAuthToken"));
+	const FConfigSection* Section = InConfigFile.FindSection(TEXT("SubmitTool.OAuthToken"));
 	FOAuthTokenParams Output;
 
 	if(Section != nullptr)
@@ -367,9 +377,9 @@ FOAuthTokenParams FSubmitToolParametersBuilder::BuildOAuthParameters()
 }
 
 
-FIncompatibleFilesParams FSubmitToolParametersBuilder::BuildIncompatibleFilesParameters()
+FIncompatibleFilesParams FSubmitToolParametersBuilder::BuildIncompatibleFilesParameters(const FConfigFile& InConfigFile)
 {
-	const FConfigSection* Section = SubmitToolConfig->FindSection(TEXT("SubmitTool.IncompatibleFiles"));
+	const FConfigSection* Section = InConfigFile.FindSection(TEXT("SubmitTool.IncompatibleFiles"));
 	FIncompatibleFilesParams Output;
 
 	if(Section != nullptr)
@@ -387,9 +397,9 @@ FIncompatibleFilesParams FSubmitToolParametersBuilder::BuildIncompatibleFilesPar
 	return Output;
 }
 
-FHordeParameters FSubmitToolParametersBuilder::BuildHordeParameters()
+FHordeParameters FSubmitToolParametersBuilder::BuildHordeParameters(const FConfigFile& InConfigFile)
 {
-	const FConfigSection* Section = SubmitToolConfig->FindSection(TEXT("SubmitTool.Horde"));
+	const FConfigSection* Section = InConfigFile.FindSection(TEXT("SubmitTool.Horde"));
 	FHordeParameters Output;
 
 	if(Section != nullptr)
@@ -406,9 +416,9 @@ FHordeParameters FSubmitToolParametersBuilder::BuildHordeParameters()
 	return Output;
 }
 
-FAutoUpdateParameters FSubmitToolParametersBuilder::BuildAutoUpdateParameters()
+FAutoUpdateParameters FSubmitToolParametersBuilder::BuildAutoUpdateParameters(const FConfigFile& InConfigFile)
 {
-	const FConfigSection* Section = SubmitToolConfig->FindSection(TEXT("SubmitTool.AutoUpdate"));
+	const FConfigSection* Section = InConfigFile.FindSection(TEXT("SubmitTool.AutoUpdate"));
 	FAutoUpdateParameters Output;
 
 	if(Section != nullptr)
diff --git a/Engine/Source/Programs/SubmitTool/Private/Parameters/SubmitToolParametersBuilder.h b/Engine/Source/Programs/SubmitTool/Private/Parameters/SubmitToolParametersBuilder.h
--- a/Engine/Source/Programs/SubmitTool/Private/Parameters/SubmitToolParametersBuilder.h
+++ b/Engine/Source/Programs/SubmitTool/Private/Parameters/SubmitToolParametersBuilder.h
@@ -12,24 +12,25 @@ public:
 	FSubmitToolParametersBuilder();
 	FSubmitToolParameters LoadConfigFromFiles();
 
+	FGeneralParameters ReadGeneralParametersFromLocalConfig();
+
 private:
 	TArray<FConfigLayer> ConfigHierarchy;
 	TArray<FString> ProjectNames;
-	FConfigFile* SubmitToolConfig = nullptr;
 private:
-	FGeneralParameters BuildGeneralParameters();
-	FJiraParameters BuildJiraParameters(); 
-	FTelemetryParameters GetTelemetryParameters();
-	FIntegrationParameters BuildIntegrationParameters();
-	TArray<FTagDefinition> BuildAvailableTags();
-	TMap<FString, FString> BuildValidators();
-	TMap<FString, FString> BuildPresubmitOperations();
-	FCopyLogParameters BuildCopyLogParameters();
-	FP4LockdownParameters BuildP4LockdownParameters();
-	FOAuthTokenParams BuildOAuthParameters();
-	FIncompatibleFilesParams BuildIncompatibleFilesParameters();
-	FHordeParameters BuildHordeParameters();
-	FAutoUpdateParameters BuildAutoUpdateParameters();
+	FGeneralParameters BuildGeneralParameters(const FConfigFile& InConfigFile);
+	FJiraParameters BuildJiraParameters(const FConfigFile& InConfigFile);
+	FTelemetryParameters GetTelemetryParameters(const FConfigFile& InConfigFile);
+	FIntegrationParameters BuildIntegrationParameters(const FConfigFile& InConfigFile);
+	TArray<FTagDefinition> BuildAvailableTags(const FConfigFile& InConfigFile);
+	TMap<FString, FString> BuildValidators(const FConfigFile& InConfigFile);
+	TMap<FString, FString> BuildPresubmitOperations(const FConfigFile& InConfigFile);
+	FCopyLogParameters BuildCopyLogParameters(const FConfigFile& InConfigFile);
+	FP4LockdownParameters BuildP4LockdownParameters(const FConfigFile& InConfigFile);
+	FOAuthTokenParams BuildOAuthParameters(const FConfigFile& InConfigFile);
+	FIncompatibleFilesParams BuildIncompatibleFilesParameters(const FConfigFile& InConfigFile);
+	FHordeParameters BuildHordeParameters(const FConfigFile& InConfigFile);
+	FAutoUpdateParameters BuildAutoUpdateParameters(const FConfigFile& InConfigFile);
 
 	
 	FString SectionToText(const FConfigSection& InSection) const;
diff --git a/Engine/Source/Programs/SubmitTool/Private/SubmitToolApp.cpp b/Engine/Source/Programs/SubmitTool/Private/SubmitToolApp.cpp
--- a/Engine/Source/Programs/SubmitTool/Private/SubmitToolApp.cpp
+++ b/Engine/Source/Programs/SubmitTool/Private/SubmitToolApp.cpp
@@ -268,29 +268,25 @@ UE::Tasks::TTask<TUniquePtr<FModelInterface>> SyncConfigAndInitializeModelInterf
 
 		TSharedPtr<ISTSourceControlService> SourceControlService = MakeShared<FSubmitToolPerforce>();
 		FSubmitToolParametersBuilder ParametersBuilder;
-		FSubmitToolParameters Parameters = ParametersBuilder.LoadConfigFromFiles();
-		if(Parameters.GeneralParameters.SyncOnStartupPaths.Num() > 0)
+		FGeneralParameters GeneralParameters = ParametersBuilder.ReadGeneralParametersFromLocalConfig();
+		if(GeneralParameters.SyncOnStartupPaths.Num() > 0)
 		{
 			TRACE_CPUPROFILER_EVENT_SCOPE(FSubmitToolApp::SyncOnStartupFiles);
 			const FString Root = FConfiguration::Substitute(TEXT("$(root)"));
 
 			TArray<FString> Files;
-			Files.Reserve(Parameters.GeneralParameters.SyncOnStartupPaths.Num());
-			for (const FString& Path : Parameters.GeneralParameters.SyncOnStartupPaths)
+			Files.Reserve(GeneralParameters.SyncOnStartupPaths.Num());
+			for (const FString& Path : GeneralParameters.SyncOnStartupPaths)
 			{
 				Files.Add(FConfiguration::SubstituteAndNormalizeFilename(Path));
 			}
 
 			SourceControlService->RunSyncCommand(TEXT("sync"), MoveTemp(Files), true);
-
-			// Reload config to be sure we use latest after sync
-			Parameters = ParametersBuilder.LoadConfigFromFiles();
 		}
 
-
 		TRACE_CPUPROFILER_EVENT_SCOPE(FSubmitToolApp::InitializeModelInterface);
 		// Create a new instance of model Interface so that UI can communicate
-		return MakeUnique<FModelInterface>(MoveTemp(Parameters), SourceControlService);
+		return MakeUnique<FModelInterface>(ParametersBuilder.LoadConfigFromFiles(), SourceControlService);
 	});
 }
 
diff --git a/Engine/Source/Programs/SubmitTool/Private/Version/AppVersionDefines.h b/Engine/Source/Programs/SubmitTool/Private/Version/AppVersionDefines.h
--- a/Engine/Source/Programs/SubmitTool/Private/Version/AppVersionDefines.h
+++ b/Engine/Source/Programs/SubmitTool/Private/Version/AppVersionDefines.h
@@ -21,7 +21,7 @@
 
 #define SUBMIT_TOOL_VERSION_MAJOR 6
 #define SUBMIT_TOOL_VERSION_MINOR 6
-#define SUBMIT_TOOL_VERSION_PATCH 1
+#define SUBMIT_TOOL_VERSION_PATCH 2
 
 #ifndef SUBMIT_TOOL_CHANGELIST
 #define SUBMIT_TOOL_CHANGELIST_STRING TEXT("NoCL")
