Fix broken cache validator skip behaviour due to a not normalized path comparison

#Synced-CL 50047099

diff --git a/Engine/Source/Programs/SubmitTool/Private/Logic/Services/CacheDataService.cpp b/Engine/Source/Programs/SubmitTool/Private/Logic/Services/CacheDataService.cpp
--- a/Engine/Source/Programs/SubmitTool/Private/Logic/Services/CacheDataService.cpp
+++ b/Engine/Source/Programs/SubmitTool/Private/Logic/Services/CacheDataService.cpp
@@ -38,7 +38,10 @@ void FCacheDataService::UpdateLastValidationForFiles(const FString& InCLId, cons
 
 	for(const FSCFileRef& File : InFiles)
 	{
-		FValidationRecord& ValidationRecord = Data.LastFileValidations.FindOrAdd(File->GetFilename()).LastValidationDateTimes.FindOrAdd(InValidatorId);
+		FString Filename = File->GetFilename();
+		FPaths::NormalizeFilename(Filename);
+
+		FValidationRecord& ValidationRecord = Data.LastFileValidations.FindOrAdd(Filename).LastValidationDateTimes.FindOrAdd(InValidatorId);
 		ValidationRecord.ValidatorConfig = InValidatorConfig;
 		ValidationRecord.LastValidationSuccess = InNewTimestamp;
 	}
