Fix SubmitToolTag not being added correctly when some tasks finished

#Synced-CL 50047099

diff --git a/Engine/Source/Programs/SubmitTool/Private/Logic/TasksService.cpp b/Engine/Source/Programs/SubmitTool/Private/Logic/TasksService.cpp
--- a/Engine/Source/Programs/SubmitTool/Private/Logic/TasksService.cpp
+++ b/Engine/Source/Programs/SubmitTool/Private/Logic/TasksService.cpp
@@ -282,7 +282,7 @@ bool FTasksService::GetIsRunSuccessful(bool bWaitForOptionalCompletes) const
 		}
 		else
 		{
-			if (bWaitForOptionalCompletes && TaskPair.Value->Definition->bRequireCompleteWhenOptional && TaskPair.Value->GetIsRunningOrQueued())
+			if (bWaitForOptionalCompletes && TaskPair.Value->Definition->bRequireCompleteWhenOptional && TaskPair.Value->GetIsRunningOrQueued() && TaskPair.Value->GetState() != EValidationStates::Not_Run)
 			{
 				return false;
 			}
@@ -395,7 +395,7 @@ bool FTasksService::Tick(float InDeltaTime)
 		}
 		else if (Task->Definition->bRequireCompleteWhenOptional)
 		{
-			bCurrentCLValidState = bCurrentCLValidState && !Task->GetIsRunningOrQueued();
+			bCurrentCLValidState = bCurrentCLValidState && !Task->GetIsRunningOrQueued() && Task->GetState() != EValidationStates::Not_Run;
 		}
 
 		bAnyTaskRunningOrQueued = bAnyTaskRunningOrQueued || Task->GetIsRunningOrQueued();
diff --git a/Engine/Source/Programs/SubmitTool/Private/Logic/Validators/ValidatorBase.h b/Engine/Source/Programs/SubmitTool/Private/Logic/Validators/ValidatorBase.h
--- a/Engine/Source/Programs/SubmitTool/Private/Logic/Validators/ValidatorBase.h
+++ b/Engine/Source/Programs/SubmitTool/Private/Logic/Validators/ValidatorBase.h
@@ -143,6 +143,7 @@ public:
 		}
 
 		State = InbAsFailed ? EValidationStates::Failed : EValidationStates::Not_Run;
+		SetValidationFinished(State);
 	}
 
 	virtual void SetQueued(bool InbForceRun = false)
diff --git a/Engine/Source/Programs/SubmitTool/Private/Models/ModelInterface.cpp b/Engine/Source/Programs/SubmitTool/Private/Models/ModelInterface.cpp
--- a/Engine/Source/Programs/SubmitTool/Private/Models/ModelInterface.cpp
+++ b/Engine/Source/Programs/SubmitTool/Private/Models/ModelInterface.cpp
@@ -87,6 +87,7 @@ FModelInterface::FModelInterface(FSubmitToolParameters&& InParameters, const TSh
 	ParsePreSubmitOperations();
 
 	OnValidationStateUpdatedHandle = ValidationService->OnTasksRunResultUpdated.Add(FOnTaskRunStateChanged::FDelegate::CreateLambda([this](bool bIsValid) {
+		ReevaluateSubmitToolTag();
 		if (bIsValid)
 		{
 			bool bOptionalFailures = false;
@@ -99,10 +100,11 @@ FModelInterface::FModelInterface(FSubmitToolParameters&& InParameters, const TSh
 			}
 
 			UE_LOG(LogSubmitTool, Log, TEXT("The required local validation has succeeded, you're ALLOWED TO SUBMIT."))
-				if (ValidationService->GetIsAnyTaskRunning())
-				{
-					UE_LOG(LogSubmitTool, Warning, TEXT("You still have optional validations running you might want to consider waiting for them to finish."))
-				}
+			if (ValidationService->GetIsAnyTaskRunning())
+			{
+				UE_LOG(LogSubmitTool, Warning, TEXT("You still have optional validations running you might want to consider waiting for them to finish."))
+			}
+
 			if (bOptionalFailures)
 			{
 				UE_LOG(LogSubmitTool, Warning, TEXT("You have optional validations that have failed, you can still proceed with the submission if you consider that these failures are not relevant. Please make sure this is the case."))
@@ -469,6 +471,7 @@ void FModelInterface::UpdateSubmitToolTag(bool InbAdd)
 			UE_LOG(LogSubmitToolDebug, Log, TEXT("Added Submit Tool tag"));
 			DescriptionCopy.Append(TEXT("\n") + SubmitToolTag);
 			ChangelistService->SetCLDescription(DescriptionCopy, true);
+			TagService->ParseCLDescription();
 		}
 	}
 	else if (HasSubmitToolTag())
@@ -793,15 +796,21 @@ void FModelInterface::StartSubmitProcess(bool bSkipShelfDialog)
 	}
 }
 
-void FModelInterface::RequestIntegration() const
+void FModelInterface::RequestIntegration()
 {
-	IntegrationService->RequestIntegration(FOnBooleanValueChanged::CreateLambda([this](bool bSuccess)
-		{
-			if (bSuccess)
+	if (IsCLValid())
+	{
+		EvaluateFailedValidatorsTag();
+		EvaluateDisabledValidatorsTag();
+		UpdateSubmitToolTag(IsCLValid());
+		IntegrationService->RequestIntegration(FOnBooleanValueChanged::CreateLambda([this](bool bSuccess)
 			{
-				ChangeState(ESubmitToolAppState::Finished);
-			}
-		}));
+				if (bSuccess)
+				{
+					ChangeState(ESubmitToolAppState::Finished);
+				}
+			}));
+	}
 }
 void FModelInterface::RefreshStateBasedOnFiles()
 {
diff --git a/Engine/Source/Programs/SubmitTool/Private/Models/ModelInterface.h b/Engine/Source/Programs/SubmitTool/Private/Models/ModelInterface.h
--- a/Engine/Source/Programs/SubmitTool/Private/Models/ModelInterface.h
+++ b/Engine/Source/Programs/SubmitTool/Private/Models/ModelInterface.h
@@ -195,7 +195,7 @@ public:
 
 	const FString& GetTitleMessage() const;
 	bool IsIntegrationRequired() const { return bIsStreamLocked; };
-	void RequestIntegration() const;
+	void RequestIntegration();
 	const TMap<FString, TSharedPtr<FIntegrationOptionBase>>& GetIntegrationOptions() { return IntegrationService->GetIntegrationOptions(); }
 	bool ValidateIntegrationOptions(bool bSilent) const	{ return IntegrationService->ValidateIntegrationOptions(bSilent); }
 
