Making sure UI validator section order is always correct

#Synced-CL 50047099

diff --git a/Engine/Source/Programs/SubmitTool/Private/Version/AppVersionDefines.h b/Engine/Source/Programs/SubmitTool/Private/Version/AppVersionDefines.h
--- a/Engine/Source/Programs/SubmitTool/Private/Version/AppVersionDefines.h
+++ b/Engine/Source/Programs/SubmitTool/Private/Version/AppVersionDefines.h
@@ -21,7 +21,7 @@
 
 #define SUBMIT_TOOL_VERSION_MAJOR 6
 #define SUBMIT_TOOL_VERSION_MINOR 6
-#define SUBMIT_TOOL_VERSION_PATCH 2
+#define SUBMIT_TOOL_VERSION_PATCH 3
 
 #ifndef SUBMIT_TOOL_CHANGELIST
 #define SUBMIT_TOOL_CHANGELIST_STRING TEXT("NoCL")
diff --git a/Engine/Source/Programs/SubmitTool/Private/View/Widgets/ValidatorsWidget.cpp b/Engine/Source/Programs/SubmitTool/Private/View/Widgets/ValidatorsWidget.cpp
--- a/Engine/Source/Programs/SubmitTool/Private/View/Widgets/ValidatorsWidget.cpp
+++ b/Engine/Source/Programs/SubmitTool/Private/View/Widgets/ValidatorsWidget.cpp
@@ -113,81 +113,84 @@ TSharedRef<SWidget> SValidatorsWidget::BuildValidatorsView(bool bListPreSubmitOp
 		Validators.FindOrAdd(DisplayName).Add(Validator);
 	}
 
-	for(const TPair<FName, TArray<TWeakPtr<const FValidatorBase>>>& Pair : Validators)
+	AddValidatorSectionToWidget(VBox, ActiveSection, bListPreSubmitOperations);
+	AddValidatorSectionToWidget(VBox, InactiveSection, bListPreSubmitOperations);
+
+	return VBox;
+}
+
+void SValidatorsWidget::AddValidatorSectionToWidget(TSharedRef<SVerticalBox> InVBox, const FName& InDisplayName, bool InbListPreSubmitOperations)
+{
+	bool bInitiallyCollapsed = false;
+	if (!(InbListPreSubmitOperations && InDisplayName.IsEqual(ActiveSection)))
 	{
-		bool bInitiallyCollapsed = false;
-		if (!(bListPreSubmitOperations && Pair.Key.IsEqual(ActiveSection)))
-		{
-			bInitiallyCollapsed = !FSubmitToolUserPrefs::Get()->UISectionExpandState.FindOrAdd(Pair.Key, false);
-		}
+		bInitiallyCollapsed = !FSubmitToolUserPrefs::Get()->UISectionExpandState.FindOrAdd(InDisplayName, false);
+	}
 
-		VBox->AddSlot()
-			.Padding(FMargin(0.f, 2.f))
-			.AutoHeight()
-			.AttachWidget
-			(
-				SNew(SExpandableArea)
-				.InitiallyCollapsed(bInitiallyCollapsed)
-				.BorderImage(FAppStyle::Get().GetBrush("Brushes.Header"))
-				.BodyBorderImage(FAppStyle::Get().GetBrush("Brushes.Recessed"))
-				.OnAreaExpansionChanged_Lambda([DisplayName = Pair.Key](bool bExpanded) { FSubmitToolUserPrefs::Get()->UISectionExpandState.FindOrAdd(DisplayName, false) = bExpanded; })
-				.BorderBackgroundColor_Lambda([&ValidatorArray = Pair.Value, &DisplayName = Pair.Key, this]
+	InVBox->AddSlot()
+	.Padding(FMargin(0.f, 2.f))
+	.AutoHeight()
+	.AttachWidget
+	(
+		SNew(SExpandableArea)
+		.InitiallyCollapsed(bInitiallyCollapsed)
+		.BorderImage(FAppStyle::Get().GetBrush("Brushes.Header"))
+		.BodyBorderImage(FAppStyle::Get().GetBrush("Brushes.Recessed"))
+		.OnAreaExpansionChanged_Lambda([DisplayName = InDisplayName](bool bExpanded) { FSubmitToolUserPrefs::Get()->UISectionExpandState.FindOrAdd(DisplayName, false) = bExpanded; })
+		.BorderBackgroundColor_Lambda([&ValidatorArray = Validators[InDisplayName], &DisplayName = InDisplayName, this]
+			{
+				if (DisplayName.IsEqual(InactiveSection))
 				{
-					if (DisplayName.IsEqual(InactiveSection))
-					{
-						return FAppStyle::GetColor("ValidatorStateSuccess");
-					}
-
-					bool bRunning = false;
-					bool bValid = true;
-					for(const TWeakPtr<const FValidatorBase>& Validator : ValidatorArray)
-					{
-						bRunning |= Validator.Pin()->GetIsRunning();
-						bValid = bValid && Validator.Pin()->GetHasPassed();
+					return FAppStyle::GetColor("ValidatorStateSuccess");
+				}
 
-						if(Validator.Pin()->GetValidatorState() == EValidationStates::Failed || 
-							Validator.Pin()->GetValidatorState() == EValidationStates::Timeout)
-						{
-							return FAppStyle::GetColor("ValidatorStateFail");
-						}
-					}
+				bool bRunning = false;
+				bool bValid = true;
+				for (const TWeakPtr<const FValidatorBase>& Validator : ValidatorArray)
+				{
+					bRunning |= Validator.Pin()->GetIsRunning();
+					bValid = bValid && Validator.Pin()->GetHasPassed();
 
-					if(bValid)
+					if (Validator.Pin()->GetValidatorState() == EValidationStates::Failed ||
+						Validator.Pin()->GetValidatorState() == EValidationStates::Timeout)
 					{
-						return FAppStyle::GetColor("ValidatorStateSuccess");
+						return FAppStyle::GetColor("ValidatorStateFail");
 					}
+				}
 
-					if(bRunning)
-					{
-						return FAppStyle::GetColor("ValidatorStateWarning");
-					}
+				if (bValid)
+				{
+					return FAppStyle::GetColor("ValidatorStateSuccess");
+				}
 
-					return FAppStyle::GetColor("ValidatorStateNormal");})
-				.HeaderPadding(FMargin(4.0f, 2.0f))
-				.Padding(1.0f)
-				.AllowAnimatedTransition(true)
-				.HeaderContent()
-				[
-					SNew(STextBlock)
-					.Text_Lambda([DisplayName = Pair.Key] {return FText::FromName(DisplayName); })
-				]
-				.BodyContent()
+				if (bRunning)
+				{
+					return FAppStyle::GetColor("ValidatorStateWarning");
+				}
+
+				return FAppStyle::GetColor("ValidatorStateNormal"); })
+		.HeaderPadding(FMargin(4.0f, 2.0f))
+		.Padding(1.0f)
+		.AllowAnimatedTransition(true)
+		.HeaderContent()
+		[
+			SNew(STextBlock)
+				.Text_Lambda([DisplayName = InDisplayName] {return FText::FromName(DisplayName); })
+		]
+		.BodyContent()
+		[
+			SNew(SBox)
+				.Padding(2.5)
 				[
-					SNew(SBox)
-					.Padding(2.5)
-					[
-						SAssignNew(ValidatorsListView, SListView<TWeakPtr<const FValidatorBase>>)
-							.SelectionMode(ESelectionMode::None)
-							.ListItemsSource(&Pair.Value)
-							.HeaderRow(ConstructHeadersRow(Pair.Key))
-							.OnGenerateRow(this, &SValidatorsWidget::GenerateRow)
-							.IsFocusable(false)
-					]
+					SAssignNew(ValidatorsListView, SListView<TWeakPtr<const FValidatorBase>>)
+						.SelectionMode(ESelectionMode::None)
+						.ListItemsSource(&Validators[InDisplayName])
+						.HeaderRow(ConstructHeadersRow(InDisplayName))
+						.OnGenerateRow(this, &SValidatorsWidget::GenerateRow)
+						.IsFocusable(false)
 				]
-			);			
-	}
-
-	return VBox;
+		]
+	);
 }
 
 void SValidatorsWidget::RefreshValidatorView(bool bListPreSubmitOperations)
diff --git a/Engine/Source/Programs/SubmitTool/Private/View/Widgets/ValidatorsWidget.h b/Engine/Source/Programs/SubmitTool/Private/View/Widgets/ValidatorsWidget.h
--- a/Engine/Source/Programs/SubmitTool/Private/View/Widgets/ValidatorsWidget.h
+++ b/Engine/Source/Programs/SubmitTool/Private/View/Widgets/ValidatorsWidget.h
@@ -43,6 +43,7 @@ public:
 private:
 
 	TSharedRef<SWidget> BuildValidatorsView(bool bListPreSubmitOperations = false);
+	void AddValidatorSectionToWidget(TSharedRef<SVerticalBox> InVBox, const FName& InDisplayName, bool InbListPreSubmitOperations);
 	void RefreshValidatorView(bool bListPreSubmitOperations = false);
 
 	FModelInterface* ModelInterface;
