Clearing keyboard focus before modifying the cl description and when losing focus to prevent race conditions with the UI

#Synced-CL 50047099

diff --git a/Engine/Source/Programs/SubmitTool/Private/Models/ModelInterface.cpp b/Engine/Source/Programs/SubmitTool/Private/Models/ModelInterface.cpp
--- a/Engine/Source/Programs/SubmitTool/Private/Models/ModelInterface.cpp
+++ b/Engine/Source/Programs/SubmitTool/Private/Models/ModelInterface.cpp
@@ -402,6 +402,7 @@ void FModelInterface::EvaluateDisabledValidatorsTag()
 
 	if (bUpdateDescription)
 	{
+		FSlateApplication::Get().ClearKeyboardFocus();
 		ChangelistService->SetCLDescription(DescriptionCopy, true);
 		TagService->ParseCLDescription();
 	}
@@ -448,6 +449,7 @@ void FModelInterface::EvaluateFailedValidatorsTag()
 
 	if (bUpdateDescription)
 	{
+		FSlateApplication::Get().ClearKeyboardFocus();
 		ChangelistService->SetCLDescription(DescriptionCopy, true);
 		TagService->ParseCLDescription();
 	}
@@ -468,6 +470,7 @@ void FModelInterface::UpdateSubmitToolTag(bool InbAdd)
 	{
 		if (!HasSubmitToolTag())
 		{
+			FSlateApplication::Get().ClearKeyboardFocus();
 			UE_LOG(LogSubmitToolDebug, Log, TEXT("Added Submit Tool tag"));
 			DescriptionCopy.Append(TEXT("\n") + SubmitToolTag);
 			ChangelistService->SetCLDescription(DescriptionCopy, true);
@@ -480,6 +483,7 @@ void FModelInterface::UpdateSubmitToolTag(bool InbAdd)
 
 		if (RemoveTagFromDescription(VersionlessTag, DescriptionCopy))
 		{
+			FSlateApplication::Get().ClearKeyboardFocus();
 			ChangelistService->SetCLDescription(DescriptionCopy, true);
 			TagService->ParseCLDescription();
 		}
diff --git a/Engine/Source/Programs/SubmitTool/Private/View/Widgets/SubmitToolWidget.cpp b/Engine/Source/Programs/SubmitTool/Private/View/Widgets/SubmitToolWidget.cpp
--- a/Engine/Source/Programs/SubmitTool/Private/View/Widgets/SubmitToolWidget.cpp
+++ b/Engine/Source/Programs/SubmitTool/Private/View/Widgets/SubmitToolWidget.cpp
@@ -68,6 +68,7 @@ void SubmitToolWidget::HandleApplicationActivationStateChanged(bool bActive)
 	}
 	else
 	{
+		FSlateApplication::Get().ClearKeyboardFocus();
 		ModelInterface->SendDescriptionToP4();
 	}
 }
