Merging Disable/EnablePlugins when they're already present in the cmdline

#Synced-CL 50398427

diff --git a/Engine/Source/Programs/SubmitTool/Private/Logic/Validators/EditorCommandletValidator.cpp b/Engine/Source/Programs/SubmitTool/Private/Logic/Validators/EditorCommandletValidator.cpp
--- a/Engine/Source/Programs/SubmitTool/Private/Logic/Validators/EditorCommandletValidator.cpp
+++ b/Engine/Source/Programs/SubmitTool/Private/Logic/Validators/EditorCommandletValidator.cpp
@@ -109,10 +109,15 @@ void FEditorCommandletValidator::GetEditorsForPaths(const TArray<FSCFileRef>& In
 			}
 		}
 
-		const FString& PluginsAdditionalArgs = GetEnableDisablePluginArguments(UProjectFiles.Value);
+		TMap<FString, TSet<FString>> PluginsAdditionalArgs = GetEnableDisablePluginArguments(UProjectFiles.Value);
 		if (!PluginsAdditionalArgs.IsEmpty())
 		{
-			EditorParams.EditorArguments.Append(PluginsAdditionalArgs);
+			for (TPair<FString, TSet<FString>>& Pair : PluginsAdditionalArgs)
+			{
+				ExtractArgumentValues(EditorParams.EditorArguments, Pair.Key, Pair.Value);
+
+				EditorParams.EditorArguments.Append(FString::Printf(TEXT(" %s%s"), *Pair.Key, *FString::Join(Pair.Value, TEXT(","))));
+			}
 		}
 
 		if (!TypedDefinition->FromEditorOnlyArguments.IsEmpty() && FCmdLineParameters::Get().Contains(FSubmitToolCmdLine::EditorFlag))
@@ -157,10 +162,15 @@ void FEditorCommandletValidator::GetEditorsForPaths(const TArray<FSCFileRef>& In
 			}
 		}
 
-		const FString& PluginsAdditionalArgs = GetEnableDisablePluginArguments(SubProjectFiles.Value);
+		TMap<FString, TSet<FString>> PluginsAdditionalArgs = GetEnableDisablePluginArguments(SubProjectFiles.Value);
 		if (!PluginsAdditionalArgs.IsEmpty())
 		{
-			EditorParams.EditorArguments.Append(PluginsAdditionalArgs);
+			for (TPair<FString, TSet<FString>>& Pair : PluginsAdditionalArgs)
+			{
+				ExtractArgumentValues(EditorParams.EditorArguments, Pair.Key, Pair.Value);
+
+				EditorParams.EditorArguments.Append(FString::Printf(TEXT(" %s%s"), *Pair.Key, *FString::Join(Pair.Value, TEXT(","))));
+			}
 		}
 
 		if (!TypedDefinition->FromEditorOnlyArguments.IsEmpty() && FCmdLineParameters::Get().Contains(FSubmitToolCmdLine::EditorFlag))
@@ -211,7 +221,7 @@ void FEditorCommandletValidator::SortFilesByProjects(const TArray<FSCFileRef>& I
 	}
 }
 
-FString FEditorCommandletValidator::GetEnableDisablePluginArguments(const TArray<FSCFileRef>& InFilesInCL) const
+TMap<FString, TSet<FString>> FEditorCommandletValidator::GetEnableDisablePluginArguments(const TArray<FSCFileRef>& InFilesInCL) const
 {
 	TMap<FSCFileRef, int8> AddRemoveBalance;
 	for (const FSCFileRef& File : InFilesInCL)
@@ -232,8 +242,8 @@ FString FEditorCommandletValidator::GetEnableDisablePluginArguments(const TArray
 	}
 
 	const FEditorCommandletValidatorDefinition* TypedDefinition = GetTypedDefinition<FEditorCommandletValidatorDefinition>();
-	TArray<FString> PluginsToEnable = TypedDefinition->EnablePlugins;
-	TArray<FString> PluginsToDisable = TypedDefinition->DisablePlugins;
+	TSet<FString> PluginsToEnable = TSet<FString>(TypedDefinition->EnablePlugins);
+	TSet<FString> PluginsToDisable = TSet<FString>(TypedDefinition->DisablePlugins);
 
 	for (const TPair<FSCFileRef, int8>& Balance : AddRemoveBalance)
 	{
@@ -247,17 +257,48 @@ FString FEditorCommandletValidator::GetEnableDisablePluginArguments(const TArray
 		}
 	}
 
-	FString CommandlineExtraArgs;
+	TMap<FString, TSet<FString>> CommandlineExtraArgs;
 
 	if (PluginsToDisable.Num() > 0)
 	{
-		CommandlineExtraArgs.Append(FString::Printf(TEXT(" -DisablePlugins=%s"), *FString::Join(PluginsToDisable, TEXT(","))));
+		CommandlineExtraArgs.Add(TEXT("-DisablePlugins="), MoveTemp(PluginsToDisable));
 	}
 
 	if (PluginsToEnable.Num() > 0)
 	{
-		CommandlineExtraArgs.Append(FString::Printf(TEXT(" -EnablePlugins=%s"), *FString::Join(PluginsToEnable, TEXT(","))));
+		CommandlineExtraArgs.Add(TEXT("-EnablePlugins="), MoveTemp(PluginsToEnable));
 	}
 
 	return CommandlineExtraArgs;
 }
+
+void FEditorCommandletValidator::ExtractArgumentValues(FString& InArguments, const FString& InKey, TSet<FString>& InOutValues, const TCHAR InDelimiter) const
+{
+	int32 Idx = InArguments.Find(InKey);
+
+	if (Idx != INDEX_NONE)
+	{
+		int32 LoopIdx = Idx + InKey.Len();
+		int32 StartIdx = LoopIdx;
+
+		while (LoopIdx != InArguments.Len())
+		{
+			if (InArguments[LoopIdx] == InDelimiter)
+			{
+				InOutValues.Add(InArguments.Mid(StartIdx, LoopIdx - StartIdx));
+				StartIdx = LoopIdx + 1;
+			}
+			else if (InArguments[LoopIdx] == TEXT(' ') || LoopIdx == InArguments.Len() - 1)
+			{
+				++LoopIdx;
+				InOutValues.Add(InArguments.Mid(StartIdx, LoopIdx - StartIdx));
+				break;
+			}
+
+			++LoopIdx;
+		}
+
+		InArguments.RemoveAt(Idx, LoopIdx - Idx);
+		InArguments.TrimEndInline();
+	}
+}
diff --git a/Engine/Source/Programs/SubmitTool/Private/Logic/Validators/EditorCommandletValidator.h b/Engine/Source/Programs/SubmitTool/Private/Logic/Validators/EditorCommandletValidator.h
--- a/Engine/Source/Programs/SubmitTool/Private/Logic/Validators/EditorCommandletValidator.h
+++ b/Engine/Source/Programs/SubmitTool/Private/Logic/Validators/EditorCommandletValidator.h
@@ -23,7 +23,9 @@ public:
 protected:
 	void GetEditorsForPaths(const TArray<FSCFileRef>& InFilteredFilesInCL, TMap<FString, FEditorParameters>& OutProjectEditorParameters) const;
 	void SortFilesByProjects(const TArray<FSCFileRef>& InFiles, TMap<FString, TArray<FSCFileRef>>& OutProjects, TMap<FString, TArray<FSCFileRef>>& OutSubProjects) const;
-	FString GetEnableDisablePluginArguments(const TArray<FSCFileRef>& InFilesInCL) const;
+	TMap<FString, TSet<FString>> GetEnableDisablePluginArguments(const TArray<FSCFileRef>& InFilesInCL) const;
+
+	void ExtractArgumentValues(FString& InArguments, const FString& InKey, TSet<FString>& InOutValues, const TCHAR InDelimiter = TCHAR(',')) const;
 
 	TMap<FString, TUniquePtr<FProcessWrapper>> Processes;
 };
