Adding support for Jira OAuth 2.0
Fix CL Description losing focus each letter

#Synced-CL 50398427

diff --git a/Engine/Source/Developer/SubmitToolCore/Public/Parameters/SubmitToolParameters.h b/Engine/Source/Developer/SubmitToolCore/Public/Parameters/SubmitToolParameters.h
--- a/Engine/Source/Developer/SubmitToolCore/Public/Parameters/SubmitToolParameters.h
+++ b/Engine/Source/Developer/SubmitToolCore/Public/Parameters/SubmitToolParameters.h
@@ -79,6 +79,15 @@ struct FJiraParameters
 	UPROPERTY()
 	bool bDisableJira = false;
 
+	UPROPERTY()
+	FString OAuthClientId;
+
+	UPROPERTY()
+	FString OAuthClientSecret;
+
+	UPROPERTY()
+	FString CloudID;
+
 	UPROPERTY()
 	FString ServerAddress;
 	
diff --git a/Engine/Source/Programs/SubmitTool/Private/Logic/CredentialsService.cpp b/Engine/Source/Programs/SubmitTool/Private/Logic/CredentialsService.cpp
--- a/Engine/Source/Programs/SubmitTool/Private/Logic/CredentialsService.cpp
+++ b/Engine/Source/Programs/SubmitTool/Private/Logic/CredentialsService.cpp
@@ -10,6 +10,7 @@
 
 #include "JsonObjectConverter.h"
 #include "HAL/FileManager.h"
+#include "HttpServerModule.h"
 #include "Misc/FileHelper.h"
 #include "Misc/Paths.h"
 #include "Misc/Base64.h"
@@ -17,8 +18,47 @@
 
 TUniquePtr<FAES::FAESKey> FCredentialsService::Key = nullptr;
 
-FCredentialsService::FCredentialsService(const FOAuthTokenParams& InOAuthParameters)
-	: Parameters(InOAuthParameters)
+
+bool FHttpServerListener::StartListening()
+{
+	if (HttpRouter.IsValid())
+	{
+		return true;
+	}
+
+	FHttpServerModule& ServerModule = FHttpServerModule::Get();
+
+	HttpRouter = ServerModule.GetHttpRouter(Port);
+
+	if (!HttpRouter.IsValid())
+	{
+		UE_LOG(LogSubmitTool, Error, TEXT("Failed to start Http Jira OAuth listener in port %d on URI %s"), Port, *URI);
+		return false;
+	}
+
+	RouteHandle = HttpRouter->BindRoute(FHttpPath(URI), Verb, Callback);
+
+	if (RouteHandle != nullptr)
+	{
+		ServerModule.StartAllListeners();
+
+	}
+
+	return RouteHandle != nullptr;
+}
+
+void FHttpServerListener::StopListening()
+{
+	if (HttpRouter.IsValid())
+	{
+		HttpRouter->UnbindRoute(RouteHandle);
+		RouteHandle.Reset();
+		HttpRouter.Reset();
+	}
+}
+
+FCredentialsService::FCredentialsService(const FSubmitToolParameters& InParameters)
+	: Parameters(InParameters)
 {
 	if(IsOIDCTokenEnabled())
 	{
@@ -183,7 +223,7 @@ void FCredentialsService::GetOIDCToken()
 					ParseOIDCTokenData(FullOutput);
 					});
 
-				OIDCProcess = MakeUnique<FProcessWrapper>(TEXT("Oidc"), FConfiguration::SubstituteAndNormalizeFilename(Parameters.OAuthTokenTool), FString::Printf(TEXT("%s"), *Parameters.OAuthArgs), OnCompleted, OutputLineProcess);
+				OIDCProcess = MakeUnique<FProcessWrapper>(TEXT("Oidc"), FConfiguration::SubstituteAndNormalizeFilename(Parameters.OAuthParameters.OAuthTokenTool), FString::Printf(TEXT("%s"), *Parameters.OAuthParameters.OAuthArgs), OnCompleted, OutputLineProcess);
 				OIDCProcess->Start(true);
 
 				if(!OIDCProcess.IsValid() || OIDCProcess->ExitCode != 0)
@@ -225,6 +265,7 @@ bool FCredentialsService::ParseOIDCTokenData(const FString& InToken)
 }
 
 constexpr int JiraCredentialDatVersion = 1;
+constexpr int JiraCredentialTokenVersion = 2;
 void FCredentialsService::SaveCredentials() const
 {
 	TArray<uint8> Bytes;
@@ -237,11 +278,22 @@ void FCredentialsService::SaveCredentials() const
 	Bytes.SetNum(NumBytesEncrypted);
 	FAES::EncryptData(Bytes.GetData(), Bytes.Num(), *Key);
 
+	TArray<uint8> BytesToken = TArray<uint8>();
+	FMemoryWriter Writer(BytesToken);
+	if (JiraOAuthToken != nullptr)
+	{
+		FJiraToken::StaticStruct()->SerializeItem(Writer, JiraOAuthToken.Get(), nullptr);
+	}
+	int32 ActualLengthToken = BytesToken.Num();
+	int32 NumBytesEncryptedToken = Align(BytesToken.Num(), FAES::AESBlockSize);
+	BytesToken.SetNum(NumBytesEncryptedToken);
+	FAES::EncryptData(BytesToken.GetData(), BytesToken.Num(), *Key);
+
 	FArchive* File = IFileManager::Get().CreateFileWriter(*GetCredentialsFilepath(), EFileWrite::FILEWRITE_EvenIfReadOnly);
 
 	if(File != nullptr)
 	{
-		int32 Version = JiraCredentialDatVersion;
+		int32 Version = JiraCredentialTokenVersion;
 		*File << Version;
 
 		*File << NumBytesEncrypted;
@@ -251,6 +303,12 @@ void FCredentialsService::SaveCredentials() const
 		File->Serialize(Bytes.GetData(), Bytes.Num());
 		Garbage = FMath::Rand();
 		*File << Garbage;
+		*File << NumBytesEncryptedToken;
+		*File << ActualLengthToken;
+		if (BytesToken.Num() > 0)
+		{
+			File->Serialize(BytesToken.GetData(), BytesToken.Num());
+		}
 
 		File->Close();
 		delete File;
@@ -285,7 +343,7 @@ void FCredentialsService::LoadCredentials()
 				*File << Version;
 
 				// Check Versions here
-				if(Version != JiraCredentialDatVersion)
+				if(Version > JiraCredentialTokenVersion)
 				{
 					UE_LOG(LogSubmitToolDebug, Warning, TEXT("Unexpected Credentials Version, aborting credentials loading."));
 					File->Close();
@@ -299,16 +357,6 @@ void FCredentialsService::LoadCredentials()
 				*File << PaddedLength;
 				*File << LengthWithoutPadding;
 
-				// see if the file has exactly the length we expect
-				// four int32 (Version, two sizes and one garbage) + the data + a final garbage int32
-				if(File->TotalSize() != (4 * sizeof(int32) + PaddedLength + sizeof(int32)))
-				{
-					UE_LOG(LogSubmitToolDebug, Warning, TEXT("Unexpected file size login invalidated"));
-					File->Close();
-					delete File;
-					return;
-				}
-
 				int32 Garbage;
 				*File << Garbage;
 
@@ -325,6 +373,25 @@ void FCredentialsService::LoadCredentials()
 					UE_LOG(LogSubmitTool, Log, TEXT("Local Credentials loaded"));
 				}
 
+				if (Version >= JiraCredentialTokenVersion)
+				{
+					*File << Garbage;
+					// Token
+					*File << PaddedLength;
+					*File << LengthWithoutPadding;
+
+					if (LengthWithoutPadding > 0)
+					{
+						DeserializedBytes.SetNum(PaddedLength);
+						File->Serialize(DeserializedBytes.GetData(), PaddedLength);
+						FAES::DecryptData(DeserializedBytes.GetData(), DeserializedBytes.Num(), *Key);
+
+						FMemoryReader Reader(DeserializedBytes);
+						FJiraToken Token;
+						FJiraToken::StaticStruct()->SerializeItem(Reader, &Token, nullptr);
+						JiraOAuthToken = MakeUnique<FJiraToken>(MoveTemp(Token));
+					}
+				}
 				File->Close();
 				delete File;
 			}
@@ -393,7 +460,7 @@ void FCredentialsService::SetLogin(const FString& InUsername, const FString& InP
 	if(newLogin != LoginString)
 	{
 		LoginString = newLogin;
-		this->SaveCredentials();
+		SaveCredentials();
 	}
 }
 
diff --git a/Engine/Source/Programs/SubmitTool/Private/Logic/CredentialsService.h b/Engine/Source/Programs/SubmitTool/Private/Logic/CredentialsService.h
--- a/Engine/Source/Programs/SubmitTool/Private/Logic/CredentialsService.h
+++ b/Engine/Source/Programs/SubmitTool/Private/Logic/CredentialsService.h
@@ -2,31 +2,69 @@
 
 #pragma once
 #include "Misc/AES.h"
+#include "IHttpRouter.h"
+#include "HttpServerRequest.h"
 #include "Parameters/SubmitToolParameters.h"
 #include "Tasks/Task.h"
 #include "Services/Interfaces/ISubmitToolService.h"
 
+#include "CredentialsService.generated.h"
+
 class FProcessWrapper;
 
-class FCredentialsService : public ISubmitToolService
+class FHttpServerListener
 {
 public:
-	FCredentialsService(const FOAuthTokenParams& InOAuthParameters);
-	virtual ~FCredentialsService();
+	FHttpServerListener(const FString& InURI, const EHttpServerRequestVerbs& InVerb, uint32 InPort, FHttpRequestHandler InCallback) : URI(InURI), Verb(InVerb), Port(InPort), Callback(InCallback)
+	{}
+	~FHttpServerListener()
+	{
+		StopListening();
+	}
+
+	bool StartListening();
+	void StopListening();
 
-	bool HasCredentials() const
+	bool IsListening() const
 	{
-		return !LoginString.IsEmpty();
+		return HttpRouter.IsValid();
 	}
 
-	bool AreCredentialsValid() const
+private:
+	TSharedPtr<IHttpRouter> HttpRouter = nullptr;
+	FHttpRouteHandle RouteHandle = nullptr;
+	FString URI;
+	EHttpServerRequestVerbs Verb;
+	uint32 Port;
+	FHttpRequestHandler Callback;
+};
+
+USTRUCT()
+struct FJiraToken
+{
+	GENERATED_BODY()
+
+	UPROPERTY()
+	FString Access_Token;
+	UPROPERTY()
+	int32 Expires_In;
+};
+
+class FCredentialsService : public ISubmitToolService
+{
+public:
+	FCredentialsService(const FSubmitToolParameters& InParameters);
+	virtual ~FCredentialsService();
+
+	bool HasJiraCredentials() const
 	{
-		return bValidatedCredentials;
+		return Parameters.JiraParameters.OAuthClientId.IsEmpty() ? !LoginString.IsEmpty() : JiraOAuthToken != nullptr;
 	}
 
-	void SetCredentialsValid(bool bValid)
+	void InvalidateCredentials()
 	{
-		bValidatedCredentials = bValid;
+		LoginString.Empty();
+		JiraOAuthToken = nullptr;
 	}
 
 	const FString& GetEncodedLoginString() const
@@ -39,7 +77,7 @@ public:
 
 	bool IsOIDCTokenEnabled()
 	{
-		return !Parameters.OAuthTokenTool.IsEmpty();
+		return !Parameters.OAuthParameters.OAuthTokenTool.IsEmpty();
 	}
 
 	bool IsTokenReady() const
@@ -64,7 +102,19 @@ public:
 
 	UE::Tasks::TTask<void> QueueWorkForToken(TFunction<void(const FString&)> InFunction);
 
+	void SetJiraToken(FJiraToken&& InJiraToken)
+	{
+		JiraOAuthToken = MakeUnique<FJiraToken>(MoveTemp(InJiraToken));
+		SaveCredentials();
+	}
+
+	const TUniquePtr<FJiraToken>& GetJiraToken() const
+	{
+		return JiraOAuthToken;
+	}
+
 private:
+	TUniquePtr<FJiraToken> JiraOAuthToken;
 	static TUniquePtr<FAES::FAESKey> Key;
 	static void LoadKey();
 	static void GenerateKey();
@@ -74,7 +124,7 @@ private:
 	TUniquePtr<FProcessWrapper> OIDCProcess;
 	FString OIDCToken;
 	FDateTime TokenExpiration;
-	const FOAuthTokenParams Parameters;
+	const FSubmitToolParameters Parameters;
 
 	void GetOIDCToken();
 	bool ParseOIDCTokenData(const FString& InToken);
diff --git a/Engine/Source/Programs/SubmitTool/Private/Logic/JiraService.cpp b/Engine/Source/Programs/SubmitTool/Private/Logic/JiraService.cpp
--- a/Engine/Source/Programs/SubmitTool/Private/Logic/JiraService.cpp
+++ b/Engine/Source/Programs/SubmitTool/Private/Logic/JiraService.cpp
@@ -9,7 +9,6 @@
 #include "Logic/PreflightService.h"
 #include "Logic/Services/SubmitToolServiceProvider.h"
 #include "Logic/Services/Interfaces/ISTSourceControlService.h"
-#include "Logic/CredentialsService.h"
 
 #include "HttpModule.h"
 #include "Interfaces/IHttpResponse.h"
@@ -27,8 +26,11 @@ FJiraService::FJiraService(const FJiraParameters& InJiraSettings, const int32 In
 {
 	if(!Definition.ServerAddress.IsEmpty())
 	{
-		FetchJiraTickets(false);
 		LoadJiraIssues();
+		if (ServiceProvider.Pin()->GetService<FCredentialsService>()->HasJiraCredentials())
+		{
+			FetchJiraTickets(true);
+		}
 	}
 }
 
@@ -41,12 +43,73 @@ FJiraService::~FJiraService()
 	OnJiraIssuesRetrievedCallback.Unbind();
 }
 
+void FJiraService::ObtainOAuth(TFunction<void(bool)> InCallback)
+{
+	if (!Definition.OAuthClientId.IsEmpty() && !bWaitingOAuth)
+	{
+		if (!HttpServerListener.IsValid() || !HttpServerListener->IsListening())
+		{
+			HttpServerListener = MakeUnique<FHttpServerListener>(TEXT("/submittool/code"), EHttpServerRequestVerbs::VERB_GET, 8451, FHttpRequestHandler::CreateLambda([this, InCallback](const FHttpServerRequest& Request, const FHttpResultCallback& OnComplete)
+			{
+				bWaitingOAuth = true;
+				const FString Html = TEXT(
+					"<html><body>"
+					"<h2>Authentication received</h2>"
+					"<p>You can return to the application.</p>"
+					"</body></html>"
+				);
+
+				TUniquePtr<FHttpServerResponse> Response = FHttpServerResponse::Create(
+					Html,
+					TEXT("text/html; charset=utf-8")
+				);
+
+				Response->Code = EHttpServerResponseCodes::Ok;
+				OnComplete(MoveTemp(Response));
+
+				if (!Request.QueryParams.Contains(TEXT("state")) || !Request.QueryParams[TEXT("state")].Equals(UserId.ToString(EGuidFormats::Digits)))
+				{
+					bWaitingOAuth = false;
+					UE_LOG(LogSubmitTool, Error, TEXT("Failed to match user code for OAuth, please try to reauthenticate"));
+					InCallback(false);
+					return false;
+				}
+
+				if (Request.QueryParams.Contains(TEXT("code")))
+				{
+					AuthCode = Request.QueryParams[TEXT("code")];
+					UE_LOG(LogSubmitToolDebug, Log, TEXT("Retrieved Jira Auth Code"));
+					ExchangeOAuthCode(InCallback);
+				}
+				else
+				{
+					bWaitingOAuth = false;
+					UE_LOG(LogSubmitTool, Error, TEXT("Could not obtain Authorization code from Atlassian."));
+					InCallback(false);
+				}
+
+				return true;
+			}));
+
+			if(!HttpServerListener->StartListening())
+			{
+				UE_LOG(LogSubmitTool, Error, TEXT("Failed to set up Jira OAuth listener."));
+				return;
+			}
+		}
+
+		FString Url = FString::Format(TEXT("https://auth.atlassian.com/authorize?audience=api.atlassian.com&client_id={0}&scope=read%3Ajira-work%20read%3Ajira-user&redirect_uri=http%3A%2F%2F127.0.0.1%3A8451%2Fsubmittool%2Fcode&state={1}&response_type=code&prompt=consent"), { Definition.OAuthClientId, UserId.ToString(EGuidFormats::Digits) });
+		FPlatformProcess::LaunchURL(*Url, nullptr, nullptr);
+		UE_LOG(LogSubmitToolDebug, Log, TEXT("Sending OAuth request for code"));
+	}
+}
+
 bool FJiraService::FetchJiraTickets(bool InForce)
 {
 	TSharedPtr<FCredentialsService> CredentialsService = ServiceProvider.Pin()->GetService<FCredentialsService>();
-	if(CredentialsService->HasCredentials() && !Definition.ServerAddress.IsEmpty() && !Definition.bDisableJira)
+	if(CredentialsService->HasJiraCredentials() && !Definition.ServerAddress.IsEmpty() && !Definition.bDisableJira)
 	{
-		if (InForce || (JiraIssues.Num() == 0 && CredentialsService->AreCredentialsValid()))
+		if (InForce || JiraIssues.Num() == 0)
 		{
 			QueryIssues();
 			return true;
@@ -63,7 +126,8 @@ void FJiraService::Reset()
 
 void FJiraService::QueryIssues()
 {
-	if(JiraRequest.IsValid() || Definition.bDisableJira)
+	TSharedPtr<FCredentialsService> CredentialsService = ServiceProvider.Pin()->GetService<FCredentialsService>();
+	if(JiraRequest.IsValid() || Definition.bDisableJira || !CredentialsService->HasJiraCredentials())
 	{
 		return;
 	}
@@ -73,14 +137,23 @@ void FJiraService::QueryIssues()
 	
 	JiraRequest->OnProcessRequestComplete().BindRaw(this, &FJiraService::QueryIssues_HttpRequestComplete);
 
-	TSharedPtr<FCredentialsService> CredentialsService = ServiceProvider.Pin()->GetService<FCredentialsService>();
-	FString Url = FString::Format(TEXT("https://{0}/rest/api/2/search?maxResults={1}&jql=assignee={2}"), { Definition.ServerAddress, this->MaxResults, CredentialsService->GetUsername()});
+
+	FString Url;
+	if (Definition.OAuthClientId.IsEmpty())
+	{
+		Url = FString::Format(TEXT("https://{0}/rest/api/2/search?maxResults={1}&jql=assignee={2}"), { Definition.ServerAddress, this->MaxResults, CredentialsService->GetUsername()});
+		JiraRequest->SetHeader(TEXT("Authorization"), TEXT("Basic ") + CredentialsService->GetEncodedLoginString());
+	}
+	else
+	{
+		Url = FString::Format(TEXT("https://api.atlassian.com/ex/jira/{0}/rest/api/2/search/jql?maxResults={1}&jql=assignee%20%3D%20currentUser%28%29%20%20AND%20type%20IN%20(Bug%2C%20Task)%20ORDER%20BY%20created&fields=key,description,summary,priority,status,issuetype,created"), { Definition.CloudID, MaxResults });
+		JiraRequest->SetHeader(TEXT("Authorization"), TEXT("Bearer ") + CredentialsService->GetJiraToken()->Access_Token);
+	}
 
 	JiraRequest->SetURL(Url);
-	JiraRequest->SetHeader(TEXT("Authorization"), TEXT("Basic ") + CredentialsService->GetEncodedLoginString());
 	JiraRequest->SetVerb(TEXT("GET"));
 	bOngoingRequest = true;
-	UE_LOG(LogSubmitToolDebug, Log, TEXT("Sending Jira request for tickets assigned to %s"), *CredentialsService->GetUsername())
+	UE_LOG(LogSubmitToolDebug, Log, TEXT("Sending Jira request for tickets"));
 	JiraRequest->ProcessRequest();
 }
 
@@ -92,7 +165,7 @@ void FJiraService::QueryIssues_HttpRequestComplete(FHttpRequestPtr HttpRequest,
 
 	if (!bSucceeded)
 	{
-		UE_LOG(LogSubmitToolDebug, Error, TEXT("Unable to retrieve JIRA issues at the moment."))
+		UE_LOG(LogSubmitToolDebug, Error, TEXT("Unable to retrieve JIRA issues at the moment."));
 		return;
 	}
 
@@ -100,7 +173,6 @@ void FJiraService::QueryIssues_HttpRequestComplete(FHttpRequestPtr HttpRequest,
 	{
 		if (EHttpResponseCodes::IsOk(HttpResponse->GetResponseCode()))
 		{
-			CredentialsService->SetCredentialsValid(true);
 			UE_LOG(LogSubmitToolDebug, Log, TEXT("Successfully connected to Jira"));
 
 			FString ResponseStr = HttpResponse->GetContentAsString();
@@ -111,12 +183,6 @@ void FJiraService::QueryIssues_HttpRequestComplete(FHttpRequestPtr HttpRequest,
 
 			if (RootJsonObject.IsValid())
 			{
-				int32 Total;
-				if (RootJsonObject->TryGetNumberField(TEXT("total"), Total))
-				{
-					this->TotalIssues = Total;
-				}
-
 				const TArray<TSharedPtr<FJsonValue>>* Issues;
 				if (RootJsonObject->TryGetArrayField(TEXT("issues"), Issues))
 				{
@@ -139,21 +205,32 @@ void FJiraService::QueryIssues_HttpRequestComplete(FHttpRequestPtr HttpRequest,
 							{
 								JiraIssues.Add(Issue.Key, Issue);
 							}
+							else
+							{
+								FString StringValue;
+								TSharedRef<TJsonWriter<TCHAR>> JsonWriter = TJsonWriterFactory<>::Create(&StringValue);
+								FJsonSerializer::Serialize(IssueObject->ToSharedRef(), JsonWriter);
+								UE_LOG(LogSubmitToolDebug, Log, TEXT("Failed to parse issue %s"), *StringValue);
+							}
+						}
+						else
+						{
+							UE_LOG(LogSubmitToolDebug, Log, TEXT("Failed to parse issue object"));
 						}
 					}
 
-					this->SaveJiraIssues();
+					SaveJiraIssues();
 				}
 			}
 		}
 		else
 		{
 			UE_LOG(LogSubmitTool, Error, TEXT("Jira Request failed with error code %d, please make sure you're logging with the right credentials. if your Okta password expired recently, make sure you log into JIRA via browser at least once."), HttpResponse->GetResponseCode());
-			CredentialsService->SetCredentialsValid(false);
+			CredentialsService->InvalidateCredentials();
 		}
 	}
 
-	OnJiraIssuesRetrievedCallback.ExecuteIfBound(EHttpResponseCodes::IsOk(HttpResponse->GetResponseCode()));
+	OnJiraIssuesRetrievedCallback.ExecuteIfBound(HttpResponse.IsValid() && EHttpResponseCodes::IsOk(HttpResponse->GetResponseCode()));
 }
 
 bool FJiraService::ParseJsonObject(const TSharedPtr<FJsonObject>* InJsonObject, FJiraIssue& OutJiraIssue) const
@@ -247,7 +324,7 @@ void FJiraService::SaveJiraIssues() const
 void FJiraService::LoadJiraIssues()
 {
 	// do not load the issues if there is no credentials
-	if (!ServiceProvider.Pin()->GetService<FCredentialsService>()->HasCredentials())
+	if (!ServiceProvider.Pin()->GetService<FCredentialsService>()->HasJiraCredentials())
 	{
 		return;
 	}
@@ -302,6 +379,78 @@ void FJiraService::LoadJiraIssues()
 	}
 }
 
+void FJiraService::ExchangeOAuthCode(TFunction<void(bool)> InCallback)
+{
+	FHttpModule& HttpModule = FModuleManager::LoadModuleChecked<FHttpModule>("HTTP");
+	FHttpRequestRef HttpRequest = HttpModule.Get().CreateRequest();
+
+	HttpRequest->SetURL(TEXT("https://auth.atlassian.com/oauth/token"));
+
+	HttpRequest->SetHeader(TEXT("Content-Type"), TEXT("application/json"));
+	HttpRequest->SetVerb(TEXT("POST"));
+
+
+	TSharedPtr<FJsonObject> RequestJson = MakeShared<FJsonObject>();
+
+	RequestJson->SetStringField(TEXT("grant_type"), TEXT("authorization_code"));
+	RequestJson->SetStringField(TEXT("client_id"), Definition.OAuthClientId);
+	RequestJson->SetStringField(TEXT("client_secret"), Definition.OAuthClientSecret);
+	RequestJson->SetStringField(TEXT("code"), AuthCode);
+	RequestJson->SetStringField(TEXT("redirect_uri"), TEXT("http://127.0.0.1:8452/submittool/token"));	
+
+	FString BodyString;
+	TSharedRef<TJsonWriter<TCHAR>> JsonWriter = TJsonWriterFactory<>::Create(&BodyString);
+	FJsonSerializer::Serialize(RequestJson.ToSharedRef(), JsonWriter);
+	HttpRequest->SetContentAsString(BodyString);
+
+	HttpRequest->OnProcessRequestComplete().BindLambda([this, InCallback](FHttpRequestPtr HttpRequest, FHttpResponsePtr HttpResponse, bool bSucceeded) {
+		AuthCode.Empty();
+		bWaitingOAuth = false;
+
+		if (!bSucceeded)
+		{			
+			UE_LOG(LogSubmitTool, Error, TEXT("Failed to exchange OAuth code for Token."))
+			InCallback(false);
+			return;
+		}
+
+		if (HttpResponse.IsValid())
+		{
+			if (EHttpResponseCodes::IsOk(HttpResponse->GetResponseCode()))
+			{
+				FString ResponseStr = HttpResponse->GetContentAsString();
+				FJiraToken Token;
+				if (FJsonObjectConverter::JsonObjectStringToUStruct<FJiraToken>(*ResponseStr, &Token))
+				{
+					ServiceProvider.Pin()->GetService<FCredentialsService>()->SetJiraToken(MoveTemp(Token));
+					UE_LOG(LogSubmitToolDebug, Log, TEXT("Obtained Jira Access Token correctly."));
+					InCallback(true);
+					return;
+				}
+				else
+				{
+					UE_LOG(LogSubmitTool, Error, TEXT("Failed to parse Jira OAuth Token."))
+					UE_LOG(LogSubmitToolDebug, Error, TEXT("%s"), *HttpResponse->GetContentAsString());
+
+				}
+			}
+			else
+			{
+				UE_LOG(LogSubmitTool, Error, TEXT("Received failed %d jira response."), HttpResponse->GetResponseCode())
+				UE_LOG(LogSubmitToolDebug, Error, TEXT("%s"), *HttpResponse->GetContentAsString());
+			}
+		}
+		else
+		{
+			UE_LOG(LogSubmitTool, Error, TEXT("Invalid jira response."));
+		}
+
+		InCallback(false);
+	});
+
+	HttpRequest->ProcessRequest();
+}
+
 const FString FJiraService::GetJiraIssuesFilepath() const
 {
 	return FPaths::Combine(FSubmitToolCoreUtils::GetLocalAppDataPath(), TEXT("SubmitTool"), TEXT("jira.issues.dat"));
@@ -314,7 +463,7 @@ void FJiraService::GetuserInfo()
 
 	HttpRequest->OnProcessRequestComplete().BindRaw(this, &FJiraService::GetuserInfo_HttpRequestComplete);
 
-	FString Url = FString::Format(TEXT("https://{0}/rest/api/2/myself}"), { Definition.ServerAddress });
+	FString Url = FString::Format(TEXT("https://{0}/rest/api/2/myself"), { Definition.ServerAddress });
 
 	HttpRequest->SetURL(Url);
 
@@ -335,7 +484,6 @@ void FJiraService::GetuserInfo_HttpRequestComplete(FHttpRequestPtr HttpRequest,
 	{
 		if (EHttpResponseCodes::IsOk(HttpResponse->GetResponseCode()))
 		{
-			ServiceProvider.Pin()->GetService<FCredentialsService>()->SetCredentialsValid(true);
 			FString ResponseStr = HttpResponse->GetContentAsString();
 
 			TSharedPtr<FJsonObject> RootJsonObject;
@@ -358,15 +506,25 @@ void FJiraService::GetIssueAndCreateServiceDeskRequest(const FString& Key, const
 	}
 
 	UE_LOG(LogSubmitTool, Log, TEXT("Requesting Information for linked Jira issue %s"), *Key);
-	// I set this up so that it gets information from the linked jira, but I think they changed their minds on what information is taken from the jira
+	
 	FHttpModule& HttpModule = FModuleManager::LoadModuleChecked<FHttpModule>("HTTP");
 	FHttpRequestRef HttpRequest = HttpModule.Get().CreateRequest();
 
 	HttpRequest->OnProcessRequestComplete().BindLambda([=, this](FHttpRequestPtr HttpRequest, FHttpResponsePtr HttpResponse, bool bSucceeded) {FJiraService::GetIssueAndCreateServiceDeskRequest_HttpRequestComplete(HttpRequest, HttpResponse, bSucceeded, Description, SwarmURL, InCurrentStream, InIntegrationOptions, OnComplete); });
-	
-	FString Url = FString::Format(TEXT("https://{0}/rest/api/2/issue/{1}"), { Definition.ServerAddress, Key });
+
+	FString Url;
+	if (Definition.OAuthClientId.IsEmpty())
+	{
+		Url = FString::Format(TEXT("https://{0}/rest/api/2/issue/{1}"), { Definition.ServerAddress, Key });
+		HttpRequest->SetHeader(TEXT("Authorization"), TEXT("Basic ") + ServiceProvider.Pin()->GetService<FCredentialsService>()->GetEncodedLoginString());
+	}
+	else
+	{
+		Url = FString::Format(TEXT("https://api.atlassian.com/ex/jira/{0}/rest/api/2/issue/{1}"), { Definition.CloudID, Key });
+		HttpRequest->SetHeader(TEXT("Authorization"), TEXT("Bearer ") + ServiceProvider.Pin()->GetService<FCredentialsService>()->GetJiraToken()->Access_Token);
+	}
+
 	HttpRequest->SetURL(Url);
-	HttpRequest->SetHeader(TEXT("Authorization"), TEXT("Basic ") + ServiceProvider.Pin()->GetService<FCredentialsService>()->GetEncodedLoginString());
 	HttpRequest->SetVerb(TEXT("GET"));
 
 	HttpRequest->ProcessRequest();
diff --git a/Engine/Source/Programs/SubmitTool/Private/Logic/JiraService.h b/Engine/Source/Programs/SubmitTool/Private/Logic/JiraService.h
--- a/Engine/Source/Programs/SubmitTool/Private/Logic/JiraService.h
+++ b/Engine/Source/Programs/SubmitTool/Private/Logic/JiraService.h
@@ -3,6 +3,7 @@
 #pragma once
 
 #include "CoreMinimal.h"
+#include "Logic/CredentialsService.h"
 #include "Models/JiraIssue.h"
 #include "Models/IntegrationOptions.h"
 #include "Interfaces/IHttpRequest.h"
@@ -11,6 +12,13 @@
 
 DECLARE_DELEGATE_OneParam(FOnJiraIssuesRetrieved, bool /*bValidResponse*/)
 
+enum EJiraServerRequest
+{
+	AuthorizationCode,
+	AccessToken
+
+};
+
 struct FJiraParameters;
 class FJsonObject;
 class FJsonValue;
@@ -22,6 +30,7 @@ public:
 	FJiraService() = delete;
 	FJiraService(const FJiraParameters& InJiraSettings, const int32 InMaxResults, TWeakPtr<FSubmitToolServiceProvider> InServiceProvider);
 	~FJiraService();
+	void ObtainOAuth(TFunction<void(bool)> InCallback);
 
 	bool FetchJiraTickets(bool InForce = false);
 	void Reset();
@@ -52,17 +61,23 @@ private:
 	void SaveJiraIssues() const;
 	void LoadJiraIssues();
 	
+	void ExchangeOAuthCode(TFunction<void(bool)> InCallback);
+
 	void CreateServiceDeskRequest_HttpRequestComplete(FHttpRequestPtr HttpRequest, FHttpResponsePtr HttpResponse, bool bSucceeded, const FOnBooleanValueChanged OnComplete);
 	void GetIssueAndCreateServiceDeskRequest_HttpRequestComplete(FHttpRequestPtr HttpRequest, FHttpResponsePtr HttpResponse, bool bSucceeded, const FString& Description, const FString& SwarmURL, const FString& InCurrentStream, const TMap<FString, TSharedPtr<FIntegrationOptionBase>>& InIntegrationOptions, const FOnBooleanValueChanged OnComplete);
 
 private:
+	FGuid UserId = FGuid::NewGuid();
+	FString AuthCode;
 	const FJiraParameters Definition;
 	FHttpRequestPtr JiraRequest = nullptr;
 	FHttpRequestPtr ServiceDeskRequest = nullptr;
 	const int32 MaxResults;
-	int32 TotalIssues;
 	TMap<FString, FJiraIssue> JiraIssues;
 
+	TUniquePtr<FHttpServerListener> HttpServerListener;
+	bool bWaitingOAuth = false;
+
 	TWeakPtr<FSubmitToolServiceProvider> ServiceProvider;
 };
 
diff --git a/Engine/Source/Programs/SubmitTool/Private/Models/ModelInterface.cpp b/Engine/Source/Programs/SubmitTool/Private/Models/ModelInterface.cpp
--- a/Engine/Source/Programs/SubmitTool/Private/Models/ModelInterface.cpp
+++ b/Engine/Source/Programs/SubmitTool/Private/Models/ModelInterface.cpp
@@ -64,7 +64,7 @@ FModelInterface::FModelInterface(FSubmitToolParameters&& InParameters, const TSh
 	ServiceProvider->RegisterService<FTasksService>(ValidationService.ToSharedRef(), TEXT("ValidationService"));
 	PresubmitOperationsService = MakeShared<FTasksService>(Parameters.PresubmitOperations, TEXT("SubmitTool.StandAlone.PresubmitOperation"));
 	ServiceProvider->RegisterService<FTasksService>(PresubmitOperationsService.ToSharedRef(), TEXT("PresubmitOperationsService"));
-	CredentialsService = MakeShared<FCredentialsService>(Parameters.OAuthParameters);
+	CredentialsService = MakeShared<FCredentialsService>(Parameters);
 	ServiceProvider->RegisterService<FCredentialsService>(CredentialsService.ToSharedRef());
 	ChangelistService = MakeShared<FChangelistService>(Parameters.GeneralParameters, ServiceProvider, CLReadyCallback, CLRefreshCallback);
 	ServiceProvider->RegisterService<FChangelistService>(ChangelistService.ToSharedRef());
@@ -76,7 +76,7 @@ FModelInterface::FModelInterface(FSubmitToolParameters&& InParameters, const TSh
 	ServiceProvider->RegisterService<FSwarmService>(SwarmService.ToSharedRef());
 	PreflightService = MakeShared<FPreflightService>(Parameters.HordeParameters, this, ServiceProvider);
 	ServiceProvider->RegisterService<FPreflightService>(PreflightService.ToSharedRef());
-	JiraService = MakeShared<FJiraService>(Parameters.JiraParameters, 256, ServiceProvider);
+	JiraService = MakeShared<FJiraService>(Parameters.JiraParameters, 512, ServiceProvider);
 	ServiceProvider->RegisterService<FJiraService>(JiraService.ToSharedRef());
 	IntegrationService = MakeShared<FIntegrationService>(Parameters.IntegrationParameters, ServiceProvider);
 	ServiceProvider->RegisterService<FIntegrationService>(IntegrationService.ToSharedRef());
@@ -483,7 +483,6 @@ void FModelInterface::UpdateSubmitToolTag(bool InbAdd)
 
 		if (RemoveTagFromDescription(VersionlessTag, DescriptionCopy))
 		{
-			FSlateApplication::Get().ClearKeyboardFocus();
 			ChangelistService->SetCLDescription(DescriptionCopy, true);
 			TagService->ParseCLDescription();
 		}
diff --git a/Engine/Source/Programs/SubmitTool/Private/Version/AppVersionDefines.h b/Engine/Source/Programs/SubmitTool/Private/Version/AppVersionDefines.h
--- a/Engine/Source/Programs/SubmitTool/Private/Version/AppVersionDefines.h
+++ b/Engine/Source/Programs/SubmitTool/Private/Version/AppVersionDefines.h
@@ -21,7 +21,7 @@
 
 #define SUBMIT_TOOL_VERSION_MAJOR 6
 #define SUBMIT_TOOL_VERSION_MINOR 6
-#define SUBMIT_TOOL_VERSION_PATCH 3
+#define SUBMIT_TOOL_VERSION_PATCH 4
 
 #ifndef SUBMIT_TOOL_CHANGELIST
 #define SUBMIT_TOOL_CHANGELIST_STRING TEXT("NoCL")
diff --git a/Engine/Source/Programs/SubmitTool/Private/View/Widgets/SJiraWidget.cpp b/Engine/Source/Programs/SubmitTool/Private/View/Widgets/SJiraWidget.cpp
--- a/Engine/Source/Programs/SubmitTool/Private/View/Widgets/SJiraWidget.cpp
+++ b/Engine/Source/Programs/SubmitTool/Private/View/Widgets/SJiraWidget.cpp
@@ -125,6 +125,7 @@ TSharedPtr<SWidget> SJiraWidget::BuildCredentialsWidget()
 			.AutoHeight()
 			[
 				SNew(SHorizontalBox)
+				.Visibility_Lambda([this]() { return ModelInterface->GetParameters().JiraParameters.OAuthClientId.IsEmpty() ? EVisibility::All : EVisibility::Collapsed; })
 				+SHorizontalBox::Slot()
 				.AutoWidth()
 				.HAlign(EHorizontalAlignment::HAlign_Center)
@@ -144,6 +145,7 @@ TSharedPtr<SWidget> SJiraWidget::BuildCredentialsWidget()
 			.AutoHeight()
 			[
 				SNew(SHorizontalBox)
+				.Visibility_Lambda([this]() { return ModelInterface->GetParameters().JiraParameters.OAuthClientId.IsEmpty() ? EVisibility::All : EVisibility::Collapsed; })
 				+SHorizontalBox::Slot()
 				.AutoWidth()
 				.HAlign(EHorizontalAlignment::HAlign_Center)
@@ -158,6 +160,16 @@ TSharedPtr<SWidget> SJiraWidget::BuildCredentialsWidget()
 				[
 					PasswordField.ToSharedRef()
 				]				
+			]			
+			+SVerticalBox::Slot()
+			.AutoHeight()
+			.HAlign(EHorizontalAlignment::HAlign_Center)
+			.VAlign(EVerticalAlignment::VAlign_Center)
+			.Padding(5)
+			[				
+				SNew(STextBlock)
+				.Visibility_Lambda([this]() { return ModelInterface->GetParameters().JiraParameters.OAuthClientId.IsEmpty() ? EVisibility::Collapsed : EVisibility::All; })
+				.Text(FText::FromString(TEXT("Authentication with Atlassian OAuth is needed.")))				
 			]
 			+SVerticalBox::Slot()
 			.AutoHeight()
@@ -402,8 +414,19 @@ TSharedPtr<SWidget> SJiraWidget::BuildIssuesWidget()
 
 FReply SJiraWidget::Login()
 {
-	ModelInterface->SetLogin(UsernameField->GetText().ToString(), PasswordField->GetText().ToString());
-	JiraService.Pin()->FetchJiraTickets(true);
+	if (ModelInterface->GetParameters().JiraParameters.OAuthClientId.IsEmpty())
+	{
+		ModelInterface->SetLogin(UsernameField->GetText().ToString(), PasswordField->GetText().ToString());
+	}
+	else
+	{
+		JiraService.Pin()->ObtainOAuth([this](bool bSuccess) {
+			if (bSuccess)
+			{
+				JiraService.Pin()->FetchJiraTickets(true);
+			}
+		});
+	}
 
 	return FReply::Handled();
 }
diff --git a/Engine/Source/Programs/SubmitTool/Private/View/Widgets/ValidatorsWidget.cpp b/Engine/Source/Programs/SubmitTool/Private/View/Widgets/ValidatorsWidget.cpp
--- a/Engine/Source/Programs/SubmitTool/Private/View/Widgets/ValidatorsWidget.cpp
+++ b/Engine/Source/Programs/SubmitTool/Private/View/Widgets/ValidatorsWidget.cpp
@@ -137,7 +137,7 @@ void SValidatorsWidget::AddValidatorSectionToWidget(TSharedRef<SVerticalBox> InV
 		.BorderImage(FAppStyle::Get().GetBrush("Brushes.Header"))
 		.BodyBorderImage(FAppStyle::Get().GetBrush("Brushes.Recessed"))
 		.OnAreaExpansionChanged_Lambda([DisplayName = InDisplayName](bool bExpanded) { FSubmitToolUserPrefs::Get()->UISectionExpandState.FindOrAdd(DisplayName, false) = bExpanded; })
-		.BorderBackgroundColor_Lambda([&ValidatorArray = Validators[InDisplayName], &DisplayName = InDisplayName, this]
+		.BorderBackgroundColor_Lambda([DisplayName = InDisplayName, this]
 			{
 				if (DisplayName.IsEqual(InactiveSection))
 				{
@@ -146,7 +146,7 @@ void SValidatorsWidget::AddValidatorSectionToWidget(TSharedRef<SVerticalBox> InV
 
 				bool bRunning = false;
 				bool bValid = true;
-				for (const TWeakPtr<const FValidatorBase>& Validator : ValidatorArray)
+				for (const TWeakPtr<const FValidatorBase>& Validator : Validators.FindOrAdd(DisplayName))
 				{
 					bRunning |= Validator.Pin()->GetIsRunning();
 					bValid = bValid && Validator.Pin()->GetHasPassed();
@@ -184,7 +184,7 @@ void SValidatorsWidget::AddValidatorSectionToWidget(TSharedRef<SVerticalBox> InV
 				[
 					SAssignNew(ValidatorsListView, SListView<TWeakPtr<const FValidatorBase>>)
 						.SelectionMode(ESelectionMode::None)
-						.ListItemsSource(&Validators[InDisplayName])
+						.ListItemsSource(&Validators.FindOrAdd(InDisplayName))
 						.HeaderRow(ConstructHeadersRow(InDisplayName))
 						.OnGenerateRow(this, &SValidatorsWidget::GenerateRow)
 						.IsFocusable(false)
diff --git a/Engine/Source/Programs/SubmitTool/SubmitTool.Build.cs b/Engine/Source/Programs/SubmitTool/SubmitTool.Build.cs
--- a/Engine/Source/Programs/SubmitTool/SubmitTool.Build.cs
+++ b/Engine/Source/Programs/SubmitTool/SubmitTool.Build.cs
@@ -34,6 +34,7 @@ public class SubmitTool : ModuleRules
 				"Settings",
 				"OutputLog",
 				"HTTP",
+				"HTTPServer",
 				"Json",
 				"JsonUtilities",
 				"Analytics",
