Adding config override for development submit tool version via preprocessor

#Synced-CL 50398427

diff --git a/Engine/Source/Programs/SubmitTool/Private/Parameters/SubmitToolParametersBuilder.cpp b/Engine/Source/Programs/SubmitTool/Private/Parameters/SubmitToolParametersBuilder.cpp
--- a/Engine/Source/Programs/SubmitTool/Private/Parameters/SubmitToolParametersBuilder.cpp
+++ b/Engine/Source/Programs/SubmitTool/Private/Parameters/SubmitToolParametersBuilder.cpp
@@ -19,6 +19,11 @@ FSubmitToolParametersBuilder::FSubmitToolParametersBuilder()
 	ConfigHierarchy.Add( { TEXT("SubmitToolBase"), TEXT("{PROJECT}/Config/{TYPE}.ini") } );
 	ConfigHierarchy.Add( { TEXT("Platform"), TEXT("{PROJECT}/Config/{PLATFORM}/{PLATFORM}{TYPE}.ini") } );
 
+#ifdef SUBMIT_TOOL_CONFIG_DEVOVERRIDE
+	ConfigHierarchy.Add({ TEXT("SubmitToolVersionDevOverride"), TEXT("{PROJECT}/Config/{TYPE}DevOverride.ini") });
+	ConfigHierarchy.Add({ TEXT("PlatformVersionDevOverride"), TEXT("{PROJECT}/Config/{PLATFORM}/{PLATFORM}{TYPE}DevOverride.ini") });
+#endif
+
 	FString RootDir;
 	FCmdLineParameters::Get().GetValue(FSubmitToolCmdLine::RootDir, RootDir);
 	FPaths::NormalizeDirectoryName(RootDir);
