[SubmitTool] Set a temporary path for -UserDir when running editor commandlet assets validation

#rnx

#Synced-CL 50325415

diff --git a/Engine/Source/Programs/SubmitTool/Private/Logic/Validators/EditorCommandletValidator.cpp b/Engine/Source/Programs/SubmitTool/Private/Logic/Validators/EditorCommandletValidator.cpp
--- a/Engine/Source/Programs/SubmitTool/Private/Logic/Validators/EditorCommandletValidator.cpp
+++ b/Engine/Source/Programs/SubmitTool/Private/Logic/Validators/EditorCommandletValidator.cpp
@@ -2,6 +2,7 @@
 
 #include "EditorCommandletValidator.h"
 
+#include "HAL/FileManager.h"
 #include "Misc/FileHelper.h"
 #include "Misc/Paths.h"
 #include "Misc/StringOutputDevice.h"
@@ -51,15 +52,40 @@ bool FEditorCommandletValidator::Validate(const FString& InCLDescription, const
 	TMap<FString, FEditorParameters> EditorParameters;
 	GetEditorsForPaths(InFilteredFilesInCL, EditorParameters);
 
+	FString TempUserDir = MakeTemporaryUserDirectory();
+
 	bool bProcessesStarted = false;
 	for (const TPair<FString, FEditorParameters>& EditorParametersSet : EditorParameters)
 	{
-		bProcessesStarted |= QueueProcess(EditorParametersSet.Key, EditorParametersSet.Value.EditorExePath, EditorParametersSet.Value.EditorArguments);
+		FString EditorArguments = EditorParametersSet.Value.EditorArguments;
+
+		// Use a temporary folder for -UserDir to avoid getting an saved editor preferences from the user when running commandlet
+		const bool bUseTempUserDir = !TempUserDir.IsEmpty() && !EditorParametersSet.Value.EditorArguments.Contains(TEXT("-UserDir"));
+		if (bUseTempUserDir)
+		{
+			EditorArguments += FString::Printf(TEXT(" -UserDir=\"%s\""), *TempUserDir);
+		}
+
+		bProcessesStarted |= QueueProcess(EditorParametersSet.Key, EditorParametersSet.Value.EditorExePath, EditorArguments);
 	}
 
 	return bProcessesStarted;
 }
 
+FString FEditorCommandletValidator::MakeTemporaryUserDirectory() const
+{
+	FString TempDir = FPaths::Combine(FPaths::EngineDir(), TEXT("Intermediate/"), TEXT("SubmitTool/"), TEXT("SavedEditor/"));
+	FString TempDirSaved = FPaths::Combine(TempDir, TEXT("Saved"));
+	IFileManager::Get().DeleteDirectory(*TempDirSaved, false, true);
+
+	if (IFileManager::Get().MakeDirectory(*TempDir, true))
+	{
+		return FPaths::ConvertRelativePathToFull(TempDir);
+	}
+
+	return FString();
+}
+
 void FEditorCommandletValidator::GetEditorsForPaths(const TArray<FSCFileRef>& InFilteredFilesInCL, TMap<FString, FEditorParameters>& OutProjectEditorParameters) const
 {
 	const FEditorCommandletValidatorDefinition* TypedDefinition = GetTypedDefinition<FEditorCommandletValidatorDefinition>();
diff --git a/Engine/Source/Programs/SubmitTool/Private/Logic/Validators/EditorCommandletValidator.h b/Engine/Source/Programs/SubmitTool/Private/Logic/Validators/EditorCommandletValidator.h
--- a/Engine/Source/Programs/SubmitTool/Private/Logic/Validators/EditorCommandletValidator.h
+++ b/Engine/Source/Programs/SubmitTool/Private/Logic/Validators/EditorCommandletValidator.h
@@ -28,4 +28,7 @@ protected:
 	void ExtractArgumentValues(FString& InArguments, const FString& InKey, TSet<FString>& InOutValues, const TCHAR InDelimiter = TCHAR(',')) const;
 
 	TMap<FString, TUniquePtr<FProcessWrapper>> Processes;
+
+private:
+	FString MakeTemporaryUserDirectory() const;
 };
