Disabling jira credentials error on startup

#Synced-CL 50398427

diff --git a/Engine/Source/Programs/SubmitTool/Private/Logic/JiraService.cpp b/Engine/Source/Programs/SubmitTool/Private/Logic/JiraService.cpp
--- a/Engine/Source/Programs/SubmitTool/Private/Logic/JiraService.cpp
+++ b/Engine/Source/Programs/SubmitTool/Private/Logic/JiraService.cpp
@@ -29,7 +29,7 @@ FJiraService::FJiraService(const FJiraParameters& InJiraSettings, const int32 In
 		LoadJiraIssues();
 		if (ServiceProvider.Pin()->GetService<FCredentialsService>()->HasJiraCredentials())
 		{
-			FetchJiraTickets(true);
+			FetchJiraTickets(true, true);
 		}
 	}
 }
@@ -104,14 +104,14 @@ void FJiraService::ObtainOAuth(TFunction<void(bool)> InCallback)
 	}
 }
 
-bool FJiraService::FetchJiraTickets(bool InForce)
+bool FJiraService::FetchJiraTickets(bool InForce, bool bInFailSilently)
 {
 	TSharedPtr<FCredentialsService> CredentialsService = ServiceProvider.Pin()->GetService<FCredentialsService>();
 	if(CredentialsService->HasJiraCredentials() && !Definition.ServerAddress.IsEmpty() && !Definition.bDisableJira)
 	{
 		if (InForce || JiraIssues.Num() == 0)
 		{
-			QueryIssues();
+			QueryIssues(bInFailSilently);
 			return true;
 		}
 	}
@@ -124,7 +124,7 @@ void FJiraService::Reset()
 	JiraIssues.Reset();
 }
 
-void FJiraService::QueryIssues()
+void FJiraService::QueryIssues(bool bInFailSilently)
 {
 	TSharedPtr<FCredentialsService> CredentialsService = ServiceProvider.Pin()->GetService<FCredentialsService>();
 	if(JiraRequest.IsValid() || Definition.bDisableJira || !CredentialsService->HasJiraCredentials())
@@ -135,8 +135,7 @@ void FJiraService::QueryIssues()
 	FHttpModule& HttpModule = FModuleManager::GetModuleChecked<FHttpModule>("HTTP");
 	JiraRequest = HttpModule.Get().CreateRequest();
 	
-	JiraRequest->OnProcessRequestComplete().BindRaw(this, &FJiraService::QueryIssues_HttpRequestComplete);
-
+	JiraRequest->OnProcessRequestComplete().BindLambda([this, bInFailSilently](FHttpRequestPtr HttpRequest, FHttpResponsePtr HttpResponse, bool bSucceeded) { QueryIssues_HttpRequestComplete(HttpRequest, HttpResponse, bSucceeded, bInFailSilently); });
 
 	FString Url;
 	if (Definition.OAuthClientId.IsEmpty())
@@ -157,7 +156,7 @@ void FJiraService::QueryIssues()
 	JiraRequest->ProcessRequest();
 }
 
-void FJiraService::QueryIssues_HttpRequestComplete(FHttpRequestPtr HttpRequest, FHttpResponsePtr HttpResponse, bool bSucceeded)
+void FJiraService::QueryIssues_HttpRequestComplete(FHttpRequestPtr HttpRequest, FHttpResponsePtr HttpResponse, bool bSucceeded, bool bFailSilently)
 {
 	JiraRequest = nullptr;
 	bOngoingRequest = false;
@@ -225,7 +224,14 @@ void FJiraService::QueryIssues_HttpRequestComplete(FHttpRequestPtr HttpRequest,
 		}
 		else
 		{
-			UE_LOG(LogSubmitTool, Error, TEXT("Jira Request failed with error code %d, please make sure you're logging with the right credentials. if your Okta password expired recently, make sure you log into JIRA via browser at least once."), HttpResponse->GetResponseCode());
+			if (bFailSilently)
+			{
+				UE_LOG(LogSubmitToolDebug, Error, TEXT("Suppressed Log: Jira Request failed with error code %d, please make sure you're logging with the right credentials. if your Okta password expired recently, make sure you log into JIRA via browser at least once."), HttpResponse->GetResponseCode());
+			}
+			else
+			{
+				UE_LOG(LogSubmitTool, Error, TEXT("Jira Request failed with error code %d, please make sure you're logging with the right credentials. if your Okta password expired recently, make sure you log into JIRA via browser at least once."), HttpResponse->GetResponseCode());
+			}
 			CredentialsService->InvalidateCredentials();
 		}
 	}
diff --git a/Engine/Source/Programs/SubmitTool/Private/Logic/JiraService.h b/Engine/Source/Programs/SubmitTool/Private/Logic/JiraService.h
--- a/Engine/Source/Programs/SubmitTool/Private/Logic/JiraService.h
+++ b/Engine/Source/Programs/SubmitTool/Private/Logic/JiraService.h
@@ -32,7 +32,7 @@ public:
 	~FJiraService();
 	void ObtainOAuth(TFunction<void(bool)> InCallback);
 
-	bool FetchJiraTickets(bool InForce = false);
+	bool FetchJiraTickets(bool InForce = false, bool bInFailSilently = false);
 	void Reset();
 
 	const TMap<FString, FJiraIssue>& GetIssues() { return this->JiraIssues; }
@@ -48,8 +48,8 @@ public:
 	void GetIssueAndCreateServiceDeskRequest(const FString& Key, const FString& Description, const FString& SwarmURL, const FString& InCurrentStream, const TMap<FString, TSharedPtr<FIntegrationOptionBase>>& InIntegrationOptions, const FOnBooleanValueChanged OnComplete);
 	void CreateServiceDeskRequest(TSharedPtr<FJsonObject> InBaseJiraJsonObject, const FString& Description, const FString& SwarmURL, const FString& InCurrentStream, const TMap<FString, TSharedPtr<FIntegrationOptionBase>>& InIntegrationOptions, const FOnBooleanValueChanged OnComplete);
 private:
-	void QueryIssues();
-	void QueryIssues_HttpRequestComplete(FHttpRequestPtr HttpRequest, FHttpResponsePtr HttpResponse, bool bSucceeded);
+	void QueryIssues(bool bInFailSilently = false);
+	void QueryIssues_HttpRequestComplete(FHttpRequestPtr HttpRequest, FHttpResponsePtr HttpResponse, bool bSucceeded, bool bFailSilently);
 
 	bool ParseJsonObject(const TSharedPtr<FJsonObject>* InJsonObject, FJiraIssue& OutJiraIssue) const;
 
