Fix crash trying to fetch jira ticket information when Jira hasn't authenticated

#Synced-CL 50398427

diff --git a/Engine/Source/Programs/SubmitTool/Private/Logic/JiraService.cpp b/Engine/Source/Programs/SubmitTool/Private/Logic/JiraService.cpp
--- a/Engine/Source/Programs/SubmitTool/Private/Logic/JiraService.cpp
+++ b/Engine/Source/Programs/SubmitTool/Private/Logic/JiraService.cpp
@@ -505,7 +505,9 @@ void FJiraService::GetuserInfo_HttpRequestComplete(FHttpRequestPtr HttpRequest,
 
 void FJiraService::GetIssueAndCreateServiceDeskRequest(const FString& Key, const FString& Description, const FString& SwarmURL, const FString& InCurrentStream, const TMap<FString, TSharedPtr<FIntegrationOptionBase>>& InIntegrationOptions, const FOnBooleanValueChanged OnComplete)
 {
-	if(Key.IsEmpty() || Key.Equals(TEXT("none"), ESearchCase::IgnoreCase))
+	TSharedPtr<FCredentialsService> CredentialsService = ServiceProvider.Pin()->GetService<FCredentialsService>();
+
+	if(Key.IsEmpty() || Key.Equals(TEXT("none"), ESearchCase::IgnoreCase) || !CredentialsService->HasJiraCredentials())
 	{
 		CreateServiceDeskRequest(TSharedPtr<FJsonObject>(), Description, SwarmURL, InCurrentStream, InIntegrationOptions, OnComplete);
 		return;
@@ -522,12 +524,15 @@ void FJiraService::GetIssueAndCreateServiceDeskRequest(const FString& Key, const
 	if (Definition.OAuthClientId.IsEmpty())
 	{
 		Url = FString::Format(TEXT("https://{0}/rest/api/2/issue/{1}"), { Definition.ServerAddress, Key });
-		HttpRequest->SetHeader(TEXT("Authorization"), TEXT("Basic ") + ServiceProvider.Pin()->GetService<FCredentialsService>()->GetEncodedLoginString());
+		HttpRequest->SetHeader(TEXT("Authorization"), TEXT("Basic ") + CredentialsService->GetEncodedLoginString());
 	}
 	else
 	{
 		Url = FString::Format(TEXT("https://api.atlassian.com/ex/jira/{0}/rest/api/2/issue/{1}"), { Definition.CloudID, Key });
-		HttpRequest->SetHeader(TEXT("Authorization"), TEXT("Bearer ") + ServiceProvider.Pin()->GetService<FCredentialsService>()->GetJiraToken()->Access_Token);
+		if (CredentialsService->GetJiraToken().IsValid())
+		{
+			HttpRequest->SetHeader(TEXT("Authorization"), TEXT("Bearer ") + ServiceProvider.Pin()->GetService<FCredentialsService>()->GetJiraToken()->Access_Token);
+		}
 	}
 
 	HttpRequest->SetURL(Url);
