Fix for double DisablePlugins due to trailing whitespace

#Synced-CL 50398427

diff --git a/Engine/Source/Programs/SubmitTool/Private/Logic/Validators/EditorCommandletValidator.cpp b/Engine/Source/Programs/SubmitTool/Private/Logic/Validators/EditorCommandletValidator.cpp
--- a/Engine/Source/Programs/SubmitTool/Private/Logic/Validators/EditorCommandletValidator.cpp
+++ b/Engine/Source/Programs/SubmitTool/Private/Logic/Validators/EditorCommandletValidator.cpp
@@ -317,7 +317,7 @@ void FEditorCommandletValidator::ExtractArgumentValues(FString& InArguments, con
 			else if (InArguments[LoopIdx] == TEXT(' ') || LoopIdx == InArguments.Len() - 1)
 			{
 				++LoopIdx;
-				InOutValues.Add(InArguments.Mid(StartIdx, LoopIdx - StartIdx));
+				InOutValues.Add(InArguments.Mid(StartIdx, LoopIdx - StartIdx).TrimEnd());
 				break;
 			}
 
