improving tag validator refresh with multi select widgets, Jira, User & Groups popups interaction

#Synced-CL 50398427

diff --git a/Engine/Source/Programs/SubmitTool/Private/View/Widgets/SJiraWidget.cpp b/Engine/Source/Programs/SubmitTool/Private/View/Widgets/SJiraWidget.cpp
--- a/Engine/Source/Programs/SubmitTool/Private/View/Widgets/SJiraWidget.cpp
+++ b/Engine/Source/Programs/SubmitTool/Private/View/Widgets/SJiraWidget.cpp
@@ -87,7 +87,7 @@ void SJiraWidget::Open(const FTag* InTargetTag)
 			.MinWidth(850)
 			.MinHeight(400);
 
-		ParentWindow->SetRequestDestroyWindowOverride(FRequestDestroyWindowOverride::CreateLambda([](const TSharedRef<SWindow>& Window) {Window->HideWindow(); }));
+		ParentWindow->SetRequestDestroyWindowOverride(FRequestDestroyWindowOverride::CreateLambda([this](const TSharedRef<SWindow>& Window) { if(ModelInterface != nullptr) { ModelInterface->ValidateCLDescription(); } Window->HideWindow(); }));
 		FSlateApplication::Get().AddWindowAsNativeChild(ParentWindow.ToSharedRef(), MainWindow.ToSharedRef(), false);
 
 		FDeprecateSlateVector2D NewPosition = MainWindow->GetPositionInScreen();
@@ -508,7 +508,11 @@ void SJiraWidget::ProcessIssueSelected(TSharedPtr<FJiraIssue> InIssue)
 
 FReply SJiraWidget::OnCloseClicked()
 {
-	this->ParentWindow->HideWindow();
+	if (ModelInterface != nullptr)
+	{
+		ModelInterface->ValidateCLDescription();
+	}
+	ParentWindow->HideWindow();
 	return FReply::Handled();
 }
 
diff --git a/Engine/Source/Programs/SubmitTool/Private/View/Widgets/SelectSourceControlGroupWidget.cpp b/Engine/Source/Programs/SubmitTool/Private/View/Widgets/SelectSourceControlGroupWidget.cpp
--- a/Engine/Source/Programs/SubmitTool/Private/View/Widgets/SelectSourceControlGroupWidget.cpp
+++ b/Engine/Source/Programs/SubmitTool/Private/View/Widgets/SelectSourceControlGroupWidget.cpp
@@ -96,7 +96,7 @@ FReply SSelectSourceControlGroupWidget::OpenDialog()
 				]
 		];
 
-	DialogWindow->SetRequestDestroyWindowOverride(FRequestDestroyWindowOverride::CreateLambda([](const TSharedRef<SWindow>& Window) { Window->HideWindow(); }));
+	DialogWindow->SetRequestDestroyWindowOverride(FRequestDestroyWindowOverride::CreateLambda([this](const TSharedRef<SWindow>& Window) { if(ModelInterface != nullptr) { ModelInterface->ValidateCLDescription(); } Window->HideWindow(); }));
 
 	TSharedRef<SHorizontalBox> SearchBox = SNew(SHorizontalBox)
 		+ SHorizontalBox::Slot()
diff --git a/Engine/Source/Programs/SubmitTool/Private/View/Widgets/SelectSourceControlUserWidget.cpp b/Engine/Source/Programs/SubmitTool/Private/View/Widgets/SelectSourceControlUserWidget.cpp
--- a/Engine/Source/Programs/SubmitTool/Private/View/Widgets/SelectSourceControlUserWidget.cpp
+++ b/Engine/Source/Programs/SubmitTool/Private/View/Widgets/SelectSourceControlUserWidget.cpp
@@ -101,7 +101,7 @@ FReply SSelectSourceControlUserWidget::OpenDialog()
 			]
 		];
 
-	DialogWindow->SetRequestDestroyWindowOverride(FRequestDestroyWindowOverride::CreateLambda([](const TSharedRef<SWindow>& Window) { Window->HideWindow(); }));
+	DialogWindow->SetRequestDestroyWindowOverride(FRequestDestroyWindowOverride::CreateLambda([this](const TSharedRef<SWindow>& Window) { if(ModelInterface != nullptr) { ModelInterface->ValidateCLDescription(); } Window->HideWindow(); }));
 
 	TSharedRef<SHorizontalBox> SearchBox = SNew(SHorizontalBox)
 		+ SHorizontalBox::Slot()
diff --git a/Engine/Source/Programs/SubmitTool/Private/View/Widgets/TagWidget.cpp b/Engine/Source/Programs/SubmitTool/Private/View/Widgets/TagWidget.cpp
--- a/Engine/Source/Programs/SubmitTool/Private/View/Widgets/TagWidget.cpp
+++ b/Engine/Source/Programs/SubmitTool/Private/View/Widgets/TagWidget.cpp
@@ -312,6 +312,12 @@ void STagWidget::Construct(const FArguments& InArgs)
 				[
 					SNew(SComboButton)
 					.IsEnabled_Lambda([this](){ return FModelInterface::GetInputEnabled() && SelectValues.Num() > 0;})
+					.OnMenuOpenChanged_Lambda([this](bool bIsOpen) {
+						if (!bIsOpen)
+						{
+							ModelInterface->ValidateCLDescription();
+						}
+					})
 					.ButtonContent()
 					[
 						SNew(STextBlock)
@@ -674,8 +680,6 @@ void STagWidget::OnSelectedChangedFromMultiselect(TSharedPtr<FString> Value)
 	{
 		ModelInterface->RemoveTag(*Tag);
 	}
-
-	ModelInterface->ValidateCLDescription();
 }
 FSlateColor STagWidget::GetPreflightGlobalColor() const
 {
@@ -754,7 +758,10 @@ void STagWidget::OnCheckboxChangedEvent(ECheckBoxState newState)
 
 void STagWidget::OnTextChanged(const FText& InText)
 {
-	ModelInterface->SetTagValues(*Tag, InText.ToString());
+	if (!Tag->GetValuesText().Equals(InText.ToString(), ESearchCase::IgnoreCase))
+	{
+		ModelInterface->SetTagValues(*Tag, InText.ToString());
+	}
 }
 
 FReply STagWidget::OnLabelClick()
