Gating new user editor prefs cmdline arg behind a config bool

#Synced-CL 50398427

diff --git a/Engine/Source/Developer/SubmitToolCore/Public/Logic/Validators/ValidatorDefinition.h b/Engine/Source/Developer/SubmitToolCore/Public/Logic/Validators/ValidatorDefinition.h
--- a/Engine/Source/Developer/SubmitToolCore/Public/Logic/Validators/ValidatorDefinition.h
+++ b/Engine/Source/Developer/SubmitToolCore/Public/Logic/Validators/ValidatorDefinition.h
@@ -520,4 +520,7 @@ struct FEditorCommandletValidatorDefinition : public FValidatorRunExecutableDefi
 
 	UPROPERTY()
 	TArray<FString> EnablePlugins;
+
+	UPROPERTY()
+	bool bIgnoreUserEditorPrefs = false;
 };
\ No newline at end of file
diff --git a/Engine/Source/Programs/SubmitTool/Private/Logic/Validators/EditorCommandletValidator.cpp b/Engine/Source/Programs/SubmitTool/Private/Logic/Validators/EditorCommandletValidator.cpp
--- a/Engine/Source/Programs/SubmitTool/Private/Logic/Validators/EditorCommandletValidator.cpp
+++ b/Engine/Source/Programs/SubmitTool/Private/Logic/Validators/EditorCommandletValidator.cpp
@@ -52,14 +52,15 @@ bool FEditorCommandletValidator::Validate(const FString& InCLDescription, const
 	TMap<FString, FEditorParameters> EditorParameters;
 	GetEditorsForPaths(InFilteredFilesInCL, EditorParameters);
 
-	FString TempUserDir = MakeTemporaryUserDirectory();
+	// Only create a temporary user directory if the validator is set up that way
+	const FString TempUserDir = GetTypedDefinition<FEditorCommandletValidatorDefinition>()->bIgnoreUserEditorPrefs ? MakeTemporaryUserDirectory() : FString();
 
 	bool bProcessesStarted = false;
 	for (const TPair<FString, FEditorParameters>& EditorParametersSet : EditorParameters)
 	{
 		FString EditorArguments = EditorParametersSet.Value.EditorArguments;
 
-		// Use a temporary folder for -UserDir to avoid getting an saved editor preferences from the user when running commandlet
+		// If we have a temporary folder for -UserDir, use it to avoid getting an saved editor preferences from the user when running commandlet
 		const bool bUseTempUserDir = !TempUserDir.IsEmpty() && !EditorParametersSet.Value.EditorArguments.Contains(TEXT("-UserDir"));
 		if (bUseTempUserDir)
 		{
