Fix incorrect Printf formatting causing compilation errors

#Synced-CL 50661505

diff --git a/Engine/Source/Programs/SubmitTool/Private/Logic/ChangelistService.cpp b/Engine/Source/Programs/SubmitTool/Private/Logic/ChangelistService.cpp
--- a/Engine/Source/Programs/SubmitTool/Private/Logic/ChangelistService.cpp
+++ b/Engine/Source/Programs/SubmitTool/Private/Logic/ChangelistService.cpp
@@ -575,11 +575,11 @@ void FChangelistService::UpdateShelvedFiles(TArray<FString> InFiles, TFunction<v
 			}
 			else if (Cmd.bDroppedConnection)
 			{
-				UE_LOG(LogSubmitToolP4, Log, TEXT("Update shelf operation was cancelled"), *GetCLID());
+				UE_LOG(LogSubmitToolP4, Log, TEXT("Update shelf operation was cancelled"));
 			}
 			else
 			{
-				UE_LOG(LogSubmitToolP4, Error, TEXT("Update shelf operation failed"), *GetCLID());
+				UE_LOG(LogSubmitToolP4, Error, TEXT("Update shelf operation failed"));
 			}
 
 			if (OnComplete != nullptr)
diff --git a/Engine/Source/Programs/SubmitTool/Private/Logic/UpdateService.cpp b/Engine/Source/Programs/SubmitTool/Private/Logic/UpdateService.cpp
--- a/Engine/Source/Programs/SubmitTool/Private/Logic/UpdateService.cpp
+++ b/Engine/Source/Programs/SubmitTool/Private/Logic/UpdateService.cpp
@@ -369,7 +369,7 @@ bool FUpdateService::DownloadLatestVersion(const FString& InDeployId, const FStr
 				}
 				else
 				{
-					DownloadErrorMessage = FString::Printf(TEXT("Unable to download latest Submit Tool deployment. Unknown Error."), HttpResponse->GetResponseCode());
+					DownloadErrorMessage = FString::Printf(TEXT("Unable to download latest Submit Tool deployment. Unknown Error."));
 					UE_LOG(LogSubmitTool, Warning, TEXT("Unable to connect to horde. Connection error, no response."));
 				}
 				return;
@@ -398,7 +398,7 @@ bool FUpdateService::DownloadLatestVersion(const FString& InDeployId, const FStr
 			}
 			else
 			{
-				DownloadErrorMessage = FString::Printf(TEXT("Unable to download latest Submit Tool deployment. Unknown Error."), HttpResponse->GetResponseCode());
+				DownloadErrorMessage = FString::Printf(TEXT("Unable to download latest Submit Tool deployment. Unknown Error."));
 				UE_LOG(LogSubmitTool, Warning, TEXT("Unable to connect to horde. Connection error, no response."));
 			}
 		});
