Added support for controlled lockdown and required tags in submit tool
Added support for dynamic lockdown files based on both commandline arguments (-EnableLocks) or validators log parsing via FRegexLockdownTrigger in FValidatorRunExecutableDefinition config
This will enable us to have client-only pseudolockdown that requires adding a tag with specific values when conditions apply. This are driven from the lockdown file.

#Synced-CL 50661505

diff --git a/Engine/Source/Developer/SubmitToolCore/Private/CommandLine/CmdLineParameters.cpp b/Engine/Source/Developer/SubmitToolCore/Private/CommandLine/CmdLineParameters.cpp
--- a/Engine/Source/Developer/SubmitToolCore/Private/CommandLine/CmdLineParameters.cpp
+++ b/Engine/Source/Developer/SubmitToolCore/Private/CommandLine/CmdLineParameters.cpp
@@ -268,7 +268,18 @@ bool FCmdLineParameters::GetValue(const FString& InKey, FString& OutValue) const
 		}
 	}
 
+	return false;
+}
 
 
+bool FCmdLineParameters::GetValueArray(const FString& InKey, TArray<FString>& OutValue, FString&& Delimiter) const
+{
+	FString StringValue;
+	if (GetValue(InKey, StringValue))
+	{
+		StringValue.ParseIntoArray(OutValue, *Delimiter);
+		return true;
+	}
+
 	return false;
-}
+}
\ No newline at end of file
diff --git a/Engine/Source/Developer/SubmitToolCore/Public/CommandLine/CmdLineParameters.h b/Engine/Source/Developer/SubmitToolCore/Public/CommandLine/CmdLineParameters.h
--- a/Engine/Source/Developer/SubmitToolCore/Public/CommandLine/CmdLineParameters.h
+++ b/Engine/Source/Developer/SubmitToolCore/Public/CommandLine/CmdLineParameters.h
@@ -14,6 +14,7 @@ namespace FSubmitToolCmdLine
 	const FString P4ChangeList = TEXT("cl");
 	const FString RootDir = TEXT("root-dir");
 	const FString EditorFlag = TEXT("from-editor");
+	const FString EnableLocks = TEXT("enablelocks");
 
 	const TArray<TSharedPtr<FCmdLineParameter>> SubmitToolCmdLineArgs =
 	{
@@ -58,6 +59,7 @@ namespace FSubmitToolCmdLine
 			FPaths::NormalizeDirectoryName(OutValue);
 		}),
 		MakeShared<FCmdLineParameter>(EditorFlag, false, TEXT("Flag indicating if submit tool is being invoked from unreal editor."), true),
+		MakeShared<FCmdLineParameter>(EnableLocks, false, TEXT("Enables lockdown files previously defined by the config by their id")),
 	};
 }
 
@@ -73,6 +75,7 @@ public:
 
 	bool Contains(const FString& InKey) const;
 	bool GetValue(const FString& InKey, FString& OutValue) const;
+	bool GetValueArray(const FString& InKey, TArray<FString>& OutValue, FString&& Delimiter = TEXT(",")) const;
 private:
 	const static FCmdLineParameters Instance;
 	void ConfigureCoreLimit() const;
diff --git a/Engine/Source/Developer/SubmitToolCore/Public/Logic/Validators/ValidatorDefinition.h b/Engine/Source/Developer/SubmitToolCore/Public/Logic/Validators/ValidatorDefinition.h
--- a/Engine/Source/Developer/SubmitToolCore/Public/Logic/Validators/ValidatorDefinition.h
+++ b/Engine/Source/Developer/SubmitToolCore/Public/Logic/Validators/ValidatorDefinition.h
@@ -274,6 +274,21 @@ struct FValidatorDefinition
 	FString PreflightTemplate;
 };
 
+USTRUCT()
+struct FRegexLockdownTrigger
+{
+	GENERATED_BODY()
+
+	UPROPERTY()
+	FString RegexPath;
+
+	TSharedPtr<FRegexPattern> BuiltRegex;
+
+	UPROPERTY()
+	FString LockdownId;
+};
+
+
 USTRUCT()
 struct FValidatorRunExecutableDefinition : public FValidatorDefinition
 {
@@ -406,6 +421,9 @@ struct FValidatorRunExecutableDefinition : public FValidatorDefinition
 	/// 
 	UPROPERTY()
 	TArray<FString> NonZeroExitCodeErrorMessages;
+
+	UPROPERTY()
+	TArray<FRegexLockdownTrigger> LogLockdownTriggers;
 };
 
 
diff --git a/Engine/Source/Developer/SubmitToolCore/Public/Parameters/SubmitToolParameters.h b/Engine/Source/Developer/SubmitToolCore/Public/Parameters/SubmitToolParameters.h
--- a/Engine/Source/Developer/SubmitToolCore/Public/Parameters/SubmitToolParameters.h
+++ b/Engine/Source/Developer/SubmitToolCore/Public/Parameters/SubmitToolParameters.h
@@ -236,6 +236,9 @@ struct FP4LockdownParameters
 	
 	UPROPERTY()
 	TArray<FString> AdditionalHardlockedPaths;
+	
+	UPROPERTY()
+	TMap<FString, FString> DynamicConfigPaths;	
 };
 
 USTRUCT()
diff --git a/Engine/Source/Programs/SubmitTool/Private/Logic/P4LockdownService.cpp b/Engine/Source/Programs/SubmitTool/Private/Logic/P4LockdownService.cpp
--- a/Engine/Source/Programs/SubmitTool/Private/Logic/P4LockdownService.cpp
+++ b/Engine/Source/Programs/SubmitTool/Private/Logic/P4LockdownService.cpp
@@ -9,6 +9,7 @@
 #include "Misc/ConfigCacheIni.h"
 #include "Misc/FileHelper.h"
 #include "Misc/Paths.h"
+#include "Models/SCFile.h"
 #include "HAL/FileManager.h"
 #include "Internationalization/Regex.h"
 #include "Configuration/Configuration.h"
@@ -25,13 +26,19 @@ FP4LockdownService::FP4LockdownService(const FP4LockdownParameters& InParameters
 
 void FP4LockdownService::DownloadAllowListData()
 {
-	if (Parameters.ConfigPaths.Num())
+	if (Parameters.ConfigPaths.Num() != 0 || Parameters.DynamicConfigPaths.Num() != 0)
 	{
 		TArray<FString> DepotPaths;
 		for (const TPair<FString, FString>& Pair : Parameters.ConfigPaths)
 		{
 			DepotPaths.Add(Pair.Value);
 		}
+
+		for (const TPair<FString, FString>& Pair : Parameters.DynamicConfigPaths)
+		{
+			DepotPaths.Add(Pair.Value);
+		}
+
 		DownloadTask = ServiceProvider.Pin()->GetService<ISTSourceControlService>()->DownloadFiles(MoveTemp(DepotPaths), DownloadedFiles);
 	}
 }
@@ -62,23 +69,31 @@ bool FP4LockdownService::WaitForAllowListData()
 	{
 		bSuccess = false;
 		int32 I = 0;
-		for (const TPair<FString, FString>& Pair : Parameters.ConfigPaths)
-		{
-			if (DownloadedFiles.Contains(Pair.Value))
+
+		auto SaveFiles = [this](const TMap<FString, FString>& InPaths) {
+			bool bSuccess = false;
+			for (const TPair<FString, FString>& Pair : InPaths)
 			{
-				FSharedBuffer FileBuffer = DownloadedFiles[Pair.Value];
-				const FString Path = GetFilePath(Pair.Key);
-				if (TUniquePtr<FArchive> File{ IFileManager::Get().CreateFileWriter(*Path, FILEWRITE_EvenIfReadOnly) })
-				{
-					File->Serialize(const_cast<void*>(FileBuffer.GetData()), FileBuffer.GetSize());
-					bSuccess = true;
-				}
-				else
+				if (DownloadedFiles.Contains(Pair.Value))
 				{
-					UE_LOG(LogSubmitToolP4, Warning, TEXT("Couldn't create lockdown file %s"), *Path)
+					FSharedBuffer FileBuffer = DownloadedFiles[Pair.Value];
+					const FString Path = GetFilePath(Pair.Key);
+					if (TUniquePtr<FArchive> File{ IFileManager::Get().CreateFileWriter(*Path, FILEWRITE_EvenIfReadOnly) })
+					{
+						File->Serialize(const_cast<void*>(FileBuffer.GetData()), FileBuffer.GetSize());
+						bSuccess = true;
+					}
+					else
+					{
+						UE_LOG(LogSubmitToolP4, Warning, TEXT("Couldn't create lockdown file %s"), *Path)
+					}
 				}
 			}
-		}
+			return bSuccess;
+		};
+
+		bSuccess |= SaveFiles(Parameters.ConfigPaths);
+		bSuccess |= SaveFiles(Parameters.DynamicConfigPaths);
 	}
 	else
 	{
@@ -88,11 +103,11 @@ bool FP4LockdownService::WaitForAllowListData()
 	return bSuccess;
 }
 
-FPathLockdownResult FP4LockdownService::ArePathsInLockdown(const TArray<FSCFileRef>& InPaths)
+FSubmitToolLockdownData FP4LockdownService::ArePathsInLockdown(const TArray<FSCFileRef>& InPaths, const TSet<FString>& InDynamicFiles)
 {	
 	TRACE_CPUPROFILER_EVENT_SCOPE(FP4LockdownService::ArePathsInLockdown);
 
-	FPathLockdownResult GlobalResult{ TSet<FString>(), TSet<FString>(), false, true };
+	FSubmitToolLockdownData GlobalResult;
 
 	if (InPaths.IsEmpty())
 	{
@@ -116,33 +131,27 @@ FPathLockdownResult FP4LockdownService::ArePathsInLockdown(const TArray<FSCFileR
 	// Finally validate the paths
 	for(const FSCFileRef& Path : InPaths)
 	{
-		FPathLockdownResult Result = IsPathInLockdown(Path->GetDepotPath());
-
-		if(Result.bIsLocked)
-		{
-			UE_LOG(LogSubmitToolP4Debug, Warning, TEXT("File %s is locked down by %s.\n%s"), *Path->GetDepotPath(), *FSubmitToolCoreUtils::StringBuilderJoin<64>(Result.GroupNames, TEXT(", ")), *FSubmitToolCoreUtils::StringBuilderJoin(Result.Messages, TEXT("\n")));
-			GlobalResult.bAllowlisted &= Result.bAllowlisted;
-			GlobalResult.bIsLocked = Result.bIsLocked;
-
-			GlobalResult.GroupNames.Append(MoveTemp(Result.GroupNames));
-			GlobalResult.Messages.Append(MoveTemp(Result.Messages));
-		}
+		GlobalResult.Append(IsPathInLockdown(Path, InDynamicFiles));
 	}
 
 	return GlobalResult;
 }
 
-FPathLockdownResult FP4LockdownService::IsPathInLockdown(const FString& InPath) const
+FSubmitToolLockdownData FP4LockdownService::IsPathInLockdown(const FSCFileRef& InPath, const TSet<FString>& InDynamicFiles) const
 {
 	FString PerforceUserName;
 	FCmdLineParameters::Get().GetValue(FSubmitToolCmdLine::P4User, PerforceUserName);
 
-	FPathLockdownResult bOverallLockdownResult { TSet<FString>(), TSet<FString>(), false, true};
+	FSubmitToolLockdownData bOverallLockdownResult;
 
 	for(const FAllowListData& data : AllowListData)
 	{
-		bool bIsLocked = false;
+		if (!data.bIsEnabled && !InDynamicFiles.Contains(data.FileId))
+		{
+			continue;
+		}
 
+		bool bIsLocked = false;
 		auto EvaluateViews = [&bIsLocked, &InPath](const TArray<TPair<bool, FString>>& Views)
 		{
 			for (const TPair<bool, FString>& View : Views)
@@ -150,7 +159,7 @@ FPathLockdownResult FP4LockdownService::IsPathInLockdown(const FString& InPath)
 				if (bIsLocked != View.Key)
 				{
 					FRegexPattern Pattern = FRegexPattern(View.Value, ERegexPatternFlags::CaseInsensitive);
-					FRegexMatcher regex = FRegexMatcher(MoveTemp(Pattern), InPath);
+					FRegexMatcher regex = FRegexMatcher(MoveTemp(Pattern), InPath->GetDepotPath());
 
 					if (regex.FindNext())
 					{
@@ -176,25 +185,43 @@ FPathLockdownResult FP4LockdownService::IsPathInLockdown(const FString& InPath)
 			const bool bAllowlisted = bIsInOverrideAllowlist || data.AllowListers.Contains(PerforceUserName);
 			const FString Message = bAllowlisted ? FString::Printf(TEXT("[Allowlisted] %s"), *data.Message) : data.Message;
 			bOverallLockdownResult.GroupNames.Add(data.GroupName);
-			bOverallLockdownResult.Messages.Add(Message);
-			bOverallLockdownResult.bIsLocked = true;
-			bOverallLockdownResult.bAllowlisted &= bAllowlisted;
+			bOverallLockdownResult.GroupMessages.Add(Message);
+			if (bOverallLockdownResult.LockdownType < data.LockdownType)
+			{
+				bOverallLockdownResult.LockdownType = data.LockdownType;
+			}
+			bOverallLockdownResult.bIsAllowlisted &= bAllowlisted;
+
+			if (bOverallLockdownResult.LockdownType == ELockdownType::Controlled && !bAllowlisted)
+			{
+				bOverallLockdownResult.RequiredTags.Add(data.ValidTag, data.AllowListers);
+			}
 		}
 	}
 
 	for(const FString& AdditionalHardlockedPath : AdditionalHardlocks)
 	{
 		FRegexPattern Pattern = FRegexPattern(AdditionalHardlockedPath, ERegexPatternFlags::CaseInsensitive);
-		FRegexMatcher regex = FRegexMatcher(MoveTemp(Pattern), InPath);
+		FRegexMatcher regex = FRegexMatcher(MoveTemp(Pattern), InPath->GetDepotPath());
 
 		if(regex.FindNext())
 		{
-			bOverallLockdownResult.bIsLocked = true;
-			bOverallLockdownResult.bAllowlisted = false;
+			bOverallLockdownResult.LockdownType = ELockdownType::Hardcore;
+			bOverallLockdownResult.bIsAllowlisted = false;
 			break;
 		}
 	}
 
+	if (bOverallLockdownResult.LockdownType != ELockdownType::None)
+	{
+		if (bOverallLockdownResult.LockdownType == ELockdownType::Hardcore)
+		{
+			bOverallLockdownResult.RequiredTags.Reset();
+		}
+
+		UE_LOG(LogSubmitToolP4Debug, Warning, TEXT("File %s is %s locked down by %s.\n%s"), *InPath->GetDepotPath(), *StaticEnum<ELockdownType>()->GetNameStringByValue(static_cast<int64>(bOverallLockdownResult.LockdownType)), *FSubmitToolCoreUtils::StringBuilderJoin<64>(bOverallLockdownResult.GroupNames, TEXT(", ")), *FSubmitToolCoreUtils::StringBuilderJoin(bOverallLockdownResult.GroupMessages, TEXT("\n")));
+	}
+
 	return bOverallLockdownResult;
 }
 
@@ -245,10 +272,16 @@ void FP4LockdownService::ParseAllowListData()
 	AllowListData.Reset();
 	OverrideData.Reset();
 
-	for(const TPair<FString, FString>& Config : Parameters.ConfigPaths)
+	ParseAllowListFiles(Parameters.ConfigPaths, true);
+	ParseAllowListFiles(Parameters.DynamicConfigPaths, false);
+}
+
+void FP4LockdownService::ParseAllowListFiles(const TMap<FString, FString>& InFiles, bool bEnabled)
+{
+	for (const TPair<FString, FString>& Config : InFiles)
 	{
 		FString Filepath = GetFilePath(Config.Key);
-		if(IFileManager::Get().FileExists(*Filepath))
+		if (IFileManager::Get().FileExists(*Filepath))
 		{
 			FConfigFile LockdownConfig;
 			LockdownConfig.bPythonConfigParserMode = true;
@@ -274,6 +307,7 @@ void FP4LockdownService::ParseAllowListData()
 				if (Status->GetSavedValue() == TEXT("hardcore"))
 				{
 					data = &AllowListData.AddDefaulted_GetRef();
+					data->LockdownType = ELockdownType::Hardcore;
 				}
 				else if (Status->GetSavedValue() == TEXT("override"))
 				{
@@ -289,17 +323,34 @@ void FP4LockdownService::ParseAllowListData()
 					TArray<FString> Sections;
 					SectionList->GetSavedValue().ParseIntoArray(Sections, TEXT(","), true);
 
-					for(FString& Section : Sections)
+					for (FString& Section : Sections)
 					{
 						Section.ToLowerInline();
 						Override.Sections.Add(MoveTemp(Section));
 					}
 				}
+				else if (Status->GetSavedValue() == TEXT("controlled"))
+				{
+					data = &AllowListData.AddDefaulted_GetRef();
+					data->LockdownType = ELockdownType::Controlled;
+
+					const FConfigValue* CustomTag = ConfigSection.Find(TEXT("customTag"));
+					if (CustomTag != nullptr)
+					{
+						data->ValidTag = CustomTag->GetSavedValue();
+					}
+					else
+					{
+						data->ValidTag = TEXT("#lockdown");
+					}
+				}
 				else
 				{
 					continue;
 				}
-
+				
+				data->FileId = Config.Key;
+				data->bIsEnabled = bEnabled;
 				data->GroupName = ConfigPair.Key;
 
 				const FConfigValue* Message = ConfigSection.Find(TEXT("message"));
@@ -311,7 +362,7 @@ void FP4LockdownService::ParseAllowListData()
 				TArray<FString> Allowlisters;
 				AllowList->GetSavedValue().ParseIntoArray(Allowlisters, TEXT(","), true);
 
-				for(FString& AllowLister : Allowlisters)
+				for (FString& AllowLister : Allowlisters)
 				{
 					AllowLister.ToLowerInline();
 					data->AllowListers.Add(MoveTemp(AllowLister));
@@ -319,15 +370,15 @@ void FP4LockdownService::ParseAllowListData()
 
 				TArray<FString> View;
 				ConfigSection.MultiFind(TEXT("view"), View, true);
-				for(FString& ViewLine : View)
+				for (FString& ViewLine : View)
 				{
-					if(ViewLine.IsEmpty())
+					if (ViewLine.IsEmpty())
 					{
 						continue;
 					}
 
 					bool bIsLocked = true;
-					if(ViewLine.StartsWith(TEXT("-"), ESearchCase::CaseSensitive))
+					if (ViewLine.StartsWith(TEXT("-"), ESearchCase::CaseSensitive))
 					{
 						bIsLocked = false;
 						ViewLine.MidInline(1);
diff --git a/Engine/Source/Programs/SubmitTool/Private/Logic/P4LockdownService.h b/Engine/Source/Programs/SubmitTool/Private/Logic/P4LockdownService.h
--- a/Engine/Source/Programs/SubmitTool/Private/Logic/P4LockdownService.h
+++ b/Engine/Source/Programs/SubmitTool/Private/Logic/P4LockdownService.h
@@ -9,17 +9,68 @@
 #include "Services/Interfaces/ISubmitToolService.h"
 #include "Memory/SharedBuffer.h"
 #include "Hash/Fnv.h"
+#include "P4LockdownService.generated.h"
+
 
 DECLARE_DELEGATE_OneParam(FOnConfigFileRetrieved, bool /*bValid*/)
 
 class FSubmitToolServiceProvider;
 
-struct FPathLockdownResult
+UENUM()
+enum class ELockdownType : int32
 {
-	TSet<FString> GroupNames;
-	TSet<FString> Messages;
-	bool bIsLocked;
-	bool bAllowlisted;
+	None = 0,	
+	Controlled,
+	Hardcore
+};
+
+struct FCaseInsensitiveStringKeyFuncs : DefaultKeyFuncs<FString, false>
+{
+	static FORCEINLINE bool Matches(const FString& A, const FString& B)
+	{
+		return A.Equals(B, ESearchCase::IgnoreCase);
+	}
+	static FORCEINLINE uint32 GetKeyHash(const FString& Key)
+	{
+		const FString Lower = Key.ToLower();
+		return UE::HashStringFNV1a32(MakeStringView(*Lower));
+	}
+};
+
+
+struct FSubmitToolLockdownData
+{
+	TSet<FString, FCaseInsensitiveStringKeyFuncs> GroupNames;
+	TSet<FString, FCaseInsensitiveStringKeyFuncs> GroupMessages;
+	ELockdownType LockdownType = ELockdownType::None;
+	TMap<FString, TSet<FString, FCaseInsensitiveStringKeyFuncs>> RequiredTags;
+	bool bIsAllowlisted = true;
+
+
+	void Append(FSubmitToolLockdownData&& OtherLockdown)
+	{
+		if (OtherLockdown.LockdownType == ELockdownType::None)
+		{
+			return;
+		}
+
+		GroupNames.Append(MoveTemp(OtherLockdown.GroupNames));
+		GroupMessages.Append(MoveTemp(OtherLockdown.GroupMessages));
+		for (auto It = OtherLockdown.RequiredTags.CreateIterator(); It; ++It)
+		{
+			if (RequiredTags.Contains(It->Key))
+			{
+				RequiredTags[It->Key].Append(MoveTemp(It->Value));
+				It.RemoveCurrent();
+			}
+		}
+		RequiredTags.Append(MoveTemp(OtherLockdown.RequiredTags));
+		bIsAllowlisted &= OtherLockdown.bIsAllowlisted;
+		if (LockdownType < OtherLockdown.LockdownType)
+		{
+			LockdownType = OtherLockdown.LockdownType;
+		}
+	}
 };
 
 class FP4LockdownService : public ISubmitToolService
@@ -28,7 +79,7 @@ public:
 
 	FP4LockdownService() = delete;
 	FP4LockdownService(const FP4LockdownParameters& InParameters, TWeakPtr<FSubmitToolServiceProvider> InServiceProvider);
-	FPathLockdownResult ArePathsInLockdown(const TArray<FSCFileRef>& InPaths);
+	FSubmitToolLockdownData ArePathsInLockdown(const TArray<FSCFileRef>& InPaths, const TSet<FString>& InDynamicFiles = TSet<FString>());
 	bool IsBlockingOperationRunning() const
 	{
 		return DownloadTask.IsValid() && !DownloadTask.IsCompleted();
@@ -36,19 +87,6 @@ public:
 
 private:
 
-	struct FCaseInsensitiveStringKeyFuncs : DefaultKeyFuncs<FString, false>
-	{
-		static FORCEINLINE bool Matches(const FString& A, const FString& B)
-		{
-			return A.Equals(B, ESearchCase::IgnoreCase);
-		}
-		static FORCEINLINE uint32 GetKeyHash(const FString& Key)
-		{
-			const FString Lower = Key.ToLower();
-			return UE::HashStringFNV1a32(MakeStringView(*Lower));
-		}
-	};
-
 	enum EAllowListState
 	{
 		Missing,
@@ -58,22 +96,27 @@ private:
 
 	struct FAllowListData
 	{
+		bool bIsEnabled = true;
+		FString FileId;
 		FString GroupName;
 		FString Message;
+		ELockdownType LockdownType = ELockdownType::None;
+		FString ValidTag;
 		TSet<FString, FCaseInsensitiveStringKeyFuncs> AllowListers;
 		TArray<TPair<bool, FString>> Views;
 	};
 
+
 	struct FOverrideData : public FAllowListData
 	{
 		TSet<FString> Sections;
 	};
-
-
-	FPathLockdownResult IsPathInLockdown(const FString& InPath) const;
+	FSubmitToolLockdownData IsPathInLockdown(const FSCFileRef& InPath, const TSet<FString>& InDynamicFiles) const;
 	void DownloadAllowListData();
 	bool WaitForAllowListData();
 	void ParseAllowListData();
+	void ParseAllowListFiles(const TMap<FString, FString>& InFiles, bool bEnabled = true);
+
 	void GetAdditionalHardlockedPaths();
 	FString GetFilePath(const FString& InConfigId);
 
diff --git a/Engine/Source/Programs/SubmitTool/Private/Logic/TagService.h b/Engine/Source/Programs/SubmitTool/Private/Logic/TagService.h
--- a/Engine/Source/Programs/SubmitTool/Private/Logic/TagService.h
+++ b/Engine/Source/Programs/SubmitTool/Private/Logic/TagService.h
@@ -15,6 +15,7 @@ public:
 
 	void ParseCLDescription();
 
+	void RegisterTransientTag(FTagDefinition Definition);
 	void ApplyTag(const FString& tagID);
 	void ApplyTag(FTag& tag);
 	void UpdateTagsInCL();
diff --git a/Engine/Source/Programs/SubmitTool/Private/Logic/Validators/ValidatorBase.cpp b/Engine/Source/Programs/SubmitTool/Private/Logic/Validators/ValidatorBase.cpp
--- a/Engine/Source/Programs/SubmitTool/Private/Logic/Validators/ValidatorBase.cpp
+++ b/Engine/Source/Programs/SubmitTool/Private/Logic/Validators/ValidatorBase.cpp
@@ -66,9 +66,16 @@ void FValidatorBase::StartValidation()
 	RunTime = 0;
 	Start = FDateTime::UtcNow();
 	State = EValidationStates::Running;
+
+	if (LockdownIds.Num() != 0)
+	{
+		LockdownIds.Reset();
+		OnLockdownIdsChanged.Broadcast(*this);
+	}
+
 	FailureReason.Reset();
-	ErrorListCache.Empty();
-	WarningListCache.Empty();
+	ErrorListCache.Reset();
+	WarningListCache.Reset();
 
 	const TSharedPtr<FChangelistService>& ChangelistService = ServiceProvider.Pin()->GetService<FChangelistService>();
 	const TSharedPtr<FTagService>& TagService = ServiceProvider.Pin()->GetService<FTagService>();
diff --git a/Engine/Source/Programs/SubmitTool/Private/Logic/Validators/ValidatorBase.h b/Engine/Source/Programs/SubmitTool/Private/Logic/Validators/ValidatorBase.h
--- a/Engine/Source/Programs/SubmitTool/Private/Logic/Validators/ValidatorBase.h
+++ b/Engine/Source/Programs/SubmitTool/Private/Logic/Validators/ValidatorBase.h
@@ -54,6 +54,7 @@ concept DerivedFromDefinition = std::is_base_of<FValidatorDefinition, T>::value;
 
 class FValidatorBase;
 DECLARE_MULTICAST_DELEGATE_OneParam(FOnValidatorFinished, const FValidatorBase& /*Validator*/)
+DECLARE_MULTICAST_DELEGATE_OneParam(FOnLockdownIdsChanged, const FValidatorBase& /*Validator*/)
 
 /**
  * @brief Base class for validator.
@@ -174,6 +175,16 @@ public:
 		return Definition->PopupMessageWhenFailed;
 	}
 
+	virtual bool IsEnabled() const
+	{
+		return !Definition->bIsDisabled;
+	}
+
+	virtual const TSet<FString>& GetLockdownIds() const
+	{
+		return LockdownIds;
+	}
+
 	const FString GetStatusText() const;
 	const EValidationStates GetValidatorState() const
 	{
@@ -214,6 +225,7 @@ public:
 	TUniquePtr<const FValidatorDefinition> Definition;
 	TMap<FString, TArray<FString>> PathsPerExtension; // Digested from Definition->IncludeFilesInDirectoryPerExtension
 	FOnValidatorFinished OnValidationFinished;
+	FOnLockdownIdsChanged OnLockdownIdsChanged;
 
 protected:
 	FName ValidatorNameID;
@@ -226,6 +238,7 @@ protected:
 	mutable TArray<FString> ErrorListCache;
 	mutable TArray<FString> WarningListCache;
 	bool bIsValidSetup = false;
+	TSet<FString> LockdownIds;
 	bool bForceRun = false;
 	const FSubmitToolParameters& SubmitToolParameters;
 
diff --git a/Engine/Source/Programs/SubmitTool/Private/Logic/Validators/ValidatorRunExecutable.cpp b/Engine/Source/Programs/SubmitTool/Private/Logic/Validators/ValidatorRunExecutable.cpp
--- a/Engine/Source/Programs/SubmitTool/Private/Logic/Validators/ValidatorRunExecutable.cpp
+++ b/Engine/Source/Programs/SubmitTool/Private/Logic/Validators/ValidatorRunExecutable.cpp
@@ -135,6 +135,11 @@ bool FValidatorRunExecutable::Activate()
 	ModifiableDefinition->BuiltRegexError = MakeShared<FRegexPattern>(ModifiableDefinition->RegexErrorParsing, ERegexPatternFlags::CaseInsensitive);
 	ModifiableDefinition->BuiltRegexWarning = MakeShared<FRegexPattern>(ModifiableDefinition->RegexWarningParsing, ERegexPatternFlags::CaseInsensitive);
 
+	for (FRegexLockdownTrigger& LockdownTrigger : ModifiableDefinition->LogLockdownTriggers)
+	{
+		LockdownTrigger.BuiltRegex = MakeShared<FRegexPattern>(LockdownTrigger.RegexPath, ERegexPatternFlags::CaseInsensitive);
+	}
+
 	static TArray<FString> ValidExtensions =
 	{
 #if PLATFORM_WINDOWS
@@ -232,6 +237,17 @@ bool FValidatorRunExecutable::QueueProcess(const FString& InId, const FString& L
 void FValidatorRunExecutable::OnProcessOutputLine(const FString& ProcessId, const FString& Line, const EProcessOutputType& OutputType)
 {
 	const FString& FormattedMessage = FString::Printf(TEXT("[%s] [%s] %s"), *GetValidatorName(), *ProcessId, *Line);
+
+	for (const FRegexLockdownTrigger& LockdownTrigger : GetTypedDefinition<FValidatorRunExecutableDefinition>()->LogLockdownTriggers)
+	{
+		FRegexMatcher Regex = FRegexMatcher(*LockdownTrigger.BuiltRegex, Line);
+		if (Regex.FindNext() && !LockdownIds.Contains(LockdownTrigger.LockdownId))
+		{
+			LockdownIds.Add(LockdownTrigger.LockdownId);
+			OnLockdownIdsChanged.Broadcast(*this);
+		}
+	}
+
 	bool bIsWarning = false;
 	if(OutputType == EProcessOutputType::ProcessError || IsLineAnError(Line, bIsWarning, ProcessesData[ProcessId].bIgnoringOutputError, ProcessesData[ProcessId].bIgnoringWarningsAsErrors))
 	{
diff --git a/Engine/Source/Programs/SubmitTool/Private/Models/ModelInterface.cpp b/Engine/Source/Programs/SubmitTool/Private/Models/ModelInterface.cpp
--- a/Engine/Source/Programs/SubmitTool/Private/Models/ModelInterface.cpp
+++ b/Engine/Source/Programs/SubmitTool/Private/Models/ModelInterface.cpp
@@ -19,6 +19,7 @@
 #include "Telemetry/TelemetryService.h"
 #include "ProfilingDebugging/CpuProfilerTrace.h"
 
+#include "CommandLine/CmdLineParameters.h"
 #include "Logic/Validators/CoreRedirectsValidator.h"
 #include "Logic/Validators/PreflightValidator.h"
 #include "Logic/Validators/TagValidator.h"
@@ -58,6 +59,10 @@ FModelInterface::FModelInterface(FSubmitToolParameters&& InParameters, const TSh
 		ServiceProvider->RegisterService<ICacheDataService>(MakeShared<FCacheDataService>(Parameters.GeneralParameters));
 	}
 
+	TArray<FString> DynamicLockdown;
+	FCmdLineParameters::Get().GetValueArray(FSubmitToolCmdLine::EnableLocks, DynamicLockdown);
+	FixedLockdownIds.Append(MoveTemp(DynamicLockdown));
+
 	SourceControlService->InitializeParameters(Parameters);
 	ServiceProvider->RegisterService<ISTSourceControlService>(SourceControlService.ToSharedRef());
 	ValidationService = MakeShared<FTasksService>(Parameters.Validators, TEXT("SubmitTool.StandAlone.Validator"));
@@ -170,6 +175,15 @@ FModelInterface::~FModelInterface()
 	ValidationService->OnSingleTaskFinished.Remove(OnSingleValidationFinishedHandle);
 	PresubmitOperationsService->OnTasksQueueFinished.Remove(OnPresubmitFinishedHandle);
 
+
+	for (const TWeakPtr<const FValidatorBase>& Task : ValidationService->GetTasks())
+	{
+		if (Task.IsValid())
+		{
+			const_cast<FValidatorBase*>(Task.Pin().Get())->OnLockdownIdsChanged.RemoveAll(this);
+		}
+	}
+
 	ServiceProvider.Reset();
 }
 
@@ -195,7 +209,7 @@ void FModelInterface::Dispose() const
 	ValidationService->StopTasks();
 }
 
-void FModelInterface::ParseValidators() const
+void FModelInterface::ParseValidators()
 {
 	TRACE_CPUPROFILER_EVENT_SCOPE(FModelInterface::ParseValidators);
 	TArray<TSharedRef<FValidatorBase>> Tasks;
@@ -271,6 +285,10 @@ void FModelInterface::ParseValidators() const
 		}
 	}
 
+	for (const TSharedRef<FValidatorBase>& Task : Tasks)
+	{
+		Task->OnLockdownIdsChanged.AddRaw(this, &FModelInterface::DynamicValidatorLockdownRefreshed);
+	}
 
 	ValidationService->InitializeTasks(Tasks);
 }
@@ -318,11 +336,12 @@ void FModelInterface::ParsePreSubmitOperations() const
 
 
 
-void FModelInterface::SetCLDescription(const FText& newDescription, bool DoNotInvalidate) const
+void FModelInterface::SetCLDescription(const FText& newDescription, bool DoNotInvalidate)
 {
 	if (ChangelistService->SetCLDescription(newDescription.ToString()))
 	{
 		TagService->ParseCLDescription();
+		CheckControlledLockdownRequirements();
 
 		if (!DoNotInvalidate)
 		{
@@ -575,6 +594,10 @@ void FModelInterface::StartSubmitProcess(bool bSkipShelfDialog)
 			PrepareSubmitCallBack.Broadcast();
 		}
 
+		if (!CheckControlledLockdownRequirements())
+		{
+			return;
+		}
 
 		TStringBuilder<1024> SubmitOverFailedValidatorsMessages;
 
@@ -772,7 +795,7 @@ void FModelInterface::StartSubmitProcess(bool bSkipShelfDialog)
 				}
 			}
 
-			ChangeState(ESubmitToolAppState::Submitting, bIsStreamLocked && bIsUserInAllowlist);
+			ChangeState(ESubmitToolAppState::Submitting, LockdownStatus.LockdownType != ELockdownType::None && LockdownStatus.bIsAllowlisted);
 			ChangelistService->DeleteShelvedFiles(
 				[this](const FSCCommand& InCmd) {
 					check(SubmitToolState == ESubmitToolAppState::Submitting);
@@ -789,7 +812,7 @@ void FModelInterface::StartSubmitProcess(bool bSkipShelfDialog)
 		}
 		else
 		{
-			ChangeState(ESubmitToolAppState::Submitting, bIsStreamLocked && bIsUserInAllowlist);
+			ChangeState(ESubmitToolAppState::Submitting, LockdownStatus.LockdownType != ELockdownType::None && LockdownStatus.bIsAllowlisted);
 			Submit();
 		}
 	}
@@ -838,13 +861,12 @@ void FModelInterface::RefreshStateBasedOnFiles()
 			{
 				// Copy the shelved file depot paths on the main thread since they are not safe to reference from other threads
 				UE::Tasks::Launch(UE_SOURCE_LOCATION, [this, ShelvedFilesDepotPaths = ChangelistService->GetShelvedFilesInCL()] {
-					FPathLockdownResult LockdownStatus = P4LockdownService->ArePathsInLockdown(ShelvedFilesDepotPaths);
-					UE::Tasks::Launch(UE_SOURCE_LOCATION, [this, LockdownStatus = MoveTemp(LockdownStatus)]
+					FSubmitToolLockdownData LockdownData = P4LockdownService->ArePathsInLockdown(ShelvedFilesDepotPaths, FixedLockdownIds);
+					UE::Tasks::Launch(UE_SOURCE_LOCATION, [this, LockdownData = MoveTemp(LockdownData)]() mutable
 					{
-						bIsStreamLocked = LockdownStatus.bIsLocked;
-						bIsUserInAllowlist = LockdownStatus.bAllowlisted;
-						TitleMessage = FSubmitToolCoreUtils::StringBuilderJoin(LockdownStatus.Messages, TEXT("\n"));
-						if (LockdownStatus.bIsLocked)
+						LockdownStatus = MoveTemp(LockdownData);
+						TitleMessage = FSubmitToolCoreUtils::StringBuilderJoin(LockdownStatus.GroupMessages, TEXT("\n"));
+						if (LockdownStatus.LockdownType == ELockdownType::Hardcore)
 						{
 							UE_LOG(LogSubmitTool, Log, TEXT("There are no local files in CL %s, Submit is disabled but you can still request an Integration with your shelved files"), *ChangelistService->GetCLID());
 							ChangeState(ESubmitToolAppState::WaitingUserInput);
@@ -895,13 +917,12 @@ void FModelInterface::RefreshStateBasedOnFiles()
 		}
 
 		// Copy the file depot paths on the main thread since they are not safe to reference from other threads
-		UE::Tasks::Launch(UE_SOURCE_LOCATION, [this, FilesInCL = ChangelistService->GetFilesInCL()] {
-			FPathLockdownResult LockdownStatus = P4LockdownService->ArePathsInLockdown(FilesInCL);
-			UE::Tasks::Launch(UE_SOURCE_LOCATION, [this, LockdownStatus = MoveTemp(LockdownStatus)]
+		UE::Tasks::Launch(UE_SOURCE_LOCATION, [this, FilesDepotPaths = ChangelistService->GetFilesInCL()] {
+			FSubmitToolLockdownData LockdownData = P4LockdownService->ArePathsInLockdown(FilesDepotPaths, FixedLockdownIds);
+			UE::Tasks::Launch(UE_SOURCE_LOCATION, [this, LockdownData = MoveTemp(LockdownData)]() mutable
 			{
-				bIsStreamLocked = LockdownStatus.bIsLocked;
-				bIsUserInAllowlist = LockdownStatus.bAllowlisted;
-				TitleMessage = FSubmitToolCoreUtils::StringBuilderJoin(LockdownStatus.Messages, TEXT("\n"));
+				LockdownStatus = MoveTemp(LockdownData);
+				TitleMessage = FSubmitToolCoreUtils::StringBuilderJoin(LockdownStatus.GroupMessages, TEXT("\n"));
 				ChangeState(ESubmitToolAppState::WaitingUserInput);
 			}, LowLevelTasks::ETaskPriority::Normal, UE::Tasks::EExtendedTaskPriority::GameThreadNormalPri);
 		});
@@ -1337,9 +1358,9 @@ bool FModelInterface::RemoveTagFromDescription(const FString& InTag, FString& In
 
 const FString& FModelInterface::GetTitleMessage() const
 {
-	if(TitleMessage.IsEmpty())
+	if (TitleMessage.IsEmpty())
 	{
-		if(FDateTime::Now().GetHour() < Parameters.GeneralParameters.EarlySubmitHour24 || FDateTime::Now().GetHour() >= Parameters.GeneralParameters.LateSubmitHour24)
+		if (FDateTime::Now().GetHour() < Parameters.GeneralParameters.EarlySubmitHour24 || FDateTime::Now().GetHour() >= Parameters.GeneralParameters.LateSubmitHour24)
 		{
 			TitleMessage = Parameters.GeneralParameters.OutOfHoursSubmitMessage.Replace(TEXT("$(DateTime)"), *FDateTime::Now().ToFormattedString(TEXT("%H:%M")));
 		}
@@ -1347,3 +1368,114 @@ const FString& FModelInterface::GetTitleMessage() const
 
 	return TitleMessage;
 }
+
+void FModelInterface::DynamicValidatorLockdownRefreshed(const FValidatorBase& InValidator)
+{
+	TRACE_CPUPROFILER_EVENT_SCOPE(FModelInterface::DynamicValidatorLockdownRefreshed);
+	bool bNeedsUpdate = false;
+	const TSet<FString>& NewSet = InValidator.GetLockdownIds();
+
+	if (DynamicValidatorLockdownIds.Contains(InValidator.GetValidatorNameId()))
+	{
+		const TSet<FString>& CurrentSet = DynamicValidatorLockdownIds[InValidator.GetValidatorNameId()];
+		bNeedsUpdate = NewSet.Num() != CurrentSet.Num() || !CurrentSet.Includes(NewSet);
+	}
+	else
+	{
+		bNeedsUpdate = !NewSet.IsEmpty();
+	}
+
+	if (bNeedsUpdate)
+	{
+		if (NewSet.IsEmpty())
+		{
+			DynamicValidatorLockdownIds.Remove(InValidator.GetValidatorNameId());
+		}
+		else
+		{
+			DynamicValidatorLockdownIds.FindOrAdd(InValidator.GetValidatorNameId()) = NewSet;
+		}
+
+		ReevaluateControlledLockdown();
+	}
+}
+
+void FModelInterface::ReevaluateControlledLockdown()
+{
+	TRACE_CPUPROFILER_EVENT_SCOPE(FModelInterface::ReevaluateControlledLockdown);
+	TSet<FString> NewLockdownIds = FixedLockdownIds;
+	for (const TPair<FName, TSet<FString>>& ValidatorLockdownIds : DynamicValidatorLockdownIds)
+	{
+		NewLockdownIds.Append(ValidatorLockdownIds.Value);
+	}
+
+	UE::Tasks::Launch(UE_SOURCE_LOCATION, [this, FilesDepotPaths = ChangelistService->GetFilesInCL(), NewLockdownIds = MoveTemp(NewLockdownIds)]() mutable {
+		FSubmitToolLockdownData LockdownData = P4LockdownService->ArePathsInLockdown(FilesDepotPaths, MoveTemp(NewLockdownIds));
+		UE::Tasks::Launch(UE_SOURCE_LOCATION, [this, LockdownData = MoveTemp(LockdownData)]() mutable
+			{
+				LockdownStatus = MoveTemp(LockdownData);
+				TitleMessage = FSubmitToolCoreUtils::StringBuilderJoin(LockdownStatus.GroupMessages, TEXT("\n"));
+				CheckControlledLockdownRequirements();
+			}, LowLevelTasks::ETaskPriority::Normal, UE::Tasks::EExtendedTaskPriority::GameThreadNormalPri);
+		});
+}
+
+bool FModelInterface::CheckControlledLockdownRequirements()
+{
+	TRACE_CPUPROFILER_EVENT_SCOPE(FModelInterface::CheckControlledLockdownRequirements);
+	if (LockdownStatus.LockdownType == ELockdownType::Controlled && !LockdownStatus.bIsAllowlisted)
+	{
+		bool bFillsRequirements = true;
+		for (const TPair<FString, TSet<FString, FCaseInsensitiveStringKeyFuncs>>& RequiredTag : LockdownStatus.RequiredTags)
+		{
+
+			FTagDefinition Definition;
+			Definition.TagId = RequiredTag.Key;
+			FTag TransientTag(Definition);
+			TransientTag.ParseTag(ChangelistService->GetCLDescription());
+
+			bool bFillsTagRequirements = TransientTag.IsEnabled() && RequiredTag.Value.Num() == 0;
+			for (const FString& Value : TransientTag.GetValues())
+			{
+				if (RequiredTag.Value.Contains(Value))
+				{
+					bFillsTagRequirements = true;
+					break;
+				}
+			}
+
+			if (!bSubmitControlledLock && !bFillsTagRequirements)
+			{
+				if (RequiredTag.Value.Num() == 0)
+				{
+					UE_LOG(LogSubmitTool, Error, TEXT("Your CL contains changes under controlled lockdown %s requiring the tag %s"), *FSubmitToolCoreUtils::StringBuilderJoin(LockdownStatus.GroupNames, TEXT(", ")), *RequiredTag.Key);
+				}
+				else
+				{
+					UE_LOG(LogSubmitTool, Error, TEXT("Your CL contains changes under controlled lockdown %s requiring the tag %s followed by one of approved allowlisters:\n%s"), *FSubmitToolCoreUtils::StringBuilderJoin(LockdownStatus.GroupNames, TEXT(", ")), *RequiredTag.Key, *FSubmitToolCoreUtils::StringBuilderJoin(RequiredTag.Value, TEXT(", ")));
+				}
+			}
+
+			bFillsRequirements &= bFillsTagRequirements;
+		}
+
+		if (bFillsRequirements || LockdownStatus.RequiredTags.Num() == 0)
+		{
+			bSubmitControlledLock = false;
+			return !bSubmitControlledLock;
+		}
+		else
+		{
+			if (!bSubmitControlledLock)
+			{
+				UE_LOG(LogSubmitTool, Error, TEXT("Additional info:\n%s"), *FSubmitToolCoreUtils::StringBuilderJoin(LockdownStatus.GroupMessages, TEXT("\n")));
+			}
+
+			bSubmitControlledLock = true;
+			return !bSubmitControlledLock;
+		}
+	}
+
+	bSubmitControlledLock = false;
+	return !bSubmitControlledLock;
+}
diff --git a/Engine/Source/Programs/SubmitTool/Private/Models/ModelInterface.h b/Engine/Source/Programs/SubmitTool/Private/Models/ModelInterface.h
--- a/Engine/Source/Programs/SubmitTool/Private/Models/ModelInterface.h
+++ b/Engine/Source/Programs/SubmitTool/Private/Models/ModelInterface.h
@@ -80,7 +80,7 @@ public:
 	FModelInterface(FSubmitToolParameters&& InParameters, const TSharedPtr<ISTSourceControlService>& InSourceControlService);
 	~FModelInterface();
 	void Dispose() const;
-	void ParseValidators() const;
+	void ParseValidators();
 	void ParsePreSubmitOperations() const;
 
 	void SetMainTab(TSharedPtr<SDockTab> InMainTab)									{ MainTab = InMainTab; }
@@ -98,7 +98,7 @@ public:
 	const TArray<const FTag*>& GetTagsArray() const									{ return TagService->GetTagsArray(); }
 	void RegisterTagUpdatedCallback(const FTagUpdated::FDelegate Callback)			{ TagService->OnTagUpdated.Add(Callback); }
 
-	void SetCLDescription(const FText& newDescription, bool DoNotInvalidate = false) const;
+	void SetCLDescription(const FText& newDescription, bool DoNotInvalidate = false);
 	void SendDescriptionToP4() const;
 	void UpdateCLFromP4Async() const;
 	const FString& GetCLDescription() const											{ return ChangelistService->GetCLDescription(); }
@@ -188,13 +188,18 @@ public:
 	TWeakPtr<FSwarmService> GetSwarmService() { return this->SwarmService; }
 	TWeakPtr<FPreflightService> GetPreflightService() { return this->PreflightService; }
 
-	bool bIsUserInAllowlist = false;
-	bool bIsStreamLocked = false;
+	FSubmitToolLockdownData LockdownStatus;
+	TSet<FString> FixedLockdownIds;
+	TMap<FName, TSet<FString>> DynamicValidatorLockdownIds;
+	void DynamicValidatorLockdownRefreshed(const FValidatorBase& InValidator);
+	void ReevaluateControlledLockdown();
+	bool CheckControlledLockdownRequirements();
 
 	const FSubmitToolParameters& GetParameters() const { return Parameters; }
 
 	const FString& GetTitleMessage() const;
-	bool IsIntegrationRequired() const { return bIsStreamLocked; };
+	bool IsIntegrationRequired() const { return LockdownStatus.LockdownType == ELockdownType::Hardcore; };
+	bool IsSubmissionBlocked() const { return bSubmitControlledLock; }
 	void RequestIntegration();
 	const TMap<FString, TSharedPtr<FIntegrationOptionBase>>& GetIntegrationOptions() { return IntegrationService->GetIntegrationOptions(); }
 	bool ValidateIntegrationOptions(bool bSilent) const	{ return IntegrationService->ValidateIntegrationOptions(bSilent); }
@@ -269,6 +274,7 @@ private:
 	TSharedPtr<FUpdateService> UpdateService;
 	TSharedPtr<FSubmitToolServiceProvider> ServiceProvider;
 	bool bPreflightQueued = false;
+	bool bSubmitControlledLock = false;
 	mutable FString TitleMessage;
 	bool Tick(float InDeltaTime);
 
diff --git a/Engine/Source/Programs/SubmitTool/Private/View/Widgets/SubmitToolWidget.cpp b/Engine/Source/Programs/SubmitTool/Private/View/Widgets/SubmitToolWidget.cpp
--- a/Engine/Source/Programs/SubmitTool/Private/View/Widgets/SubmitToolWidget.cpp
+++ b/Engine/Source/Programs/SubmitTool/Private/View/Widgets/SubmitToolWidget.cpp
@@ -520,7 +520,7 @@ TSharedRef<SHorizontalBox> SubmitToolWidget::BuildButtonRow()
 		[
 			SNew(SButton)
 			.Text(FText::FromString(TEXT("Open Integration Window")))
-			.Visibility_Lambda([this] { return ModelInterface->IsIntegrationRequired() && ModelInterface->bIsUserInAllowlist ? EVisibility::Visible : EVisibility::Collapsed; })
+			.Visibility_Lambda([this] { return ModelInterface->IsIntegrationRequired() && ModelInterface->LockdownStatus.bIsAllowlisted ? EVisibility::Visible : EVisibility::Collapsed; })
 			.OnClicked_Lambda([this] { IntegrationWidget->Open(); return FReply::Handled(); })
 			.IsEnabled_Lambda([&ModelInterface = ModelInterface]
 				{
@@ -552,7 +552,7 @@ TSharedRef<SHorizontalBox> SubmitToolWidget::BuildButtonRow()
 					SNew(SButton)
 					.IsEnabled_Lambda([&ModelInterface = ModelInterface] 
 					{
-						if (ModelInterface->IsPreflightRequestInProgress())
+						if (ModelInterface->IsPreflightRequestInProgress() || ModelInterface->IsSubmissionBlocked())
 						{
 							return false;
 						}
@@ -562,7 +562,7 @@ TSharedRef<SHorizontalBox> SubmitToolWidget::BuildButtonRow()
 							return true;
 						}
 					
-						if (ModelInterface->GetState() == ESubmitToolAppState::WaitingUserInput || (ModelInterface->IsIntegrationRequired() && ModelInterface->bIsUserInAllowlist))
+						if (ModelInterface->GetState() == ESubmitToolAppState::WaitingUserInput || (ModelInterface->IsIntegrationRequired() && ModelInterface->LockdownStatus.bIsAllowlisted))
 						{
 							return ModelInterface->IsCLValid();
 						}
@@ -836,7 +836,7 @@ FReply SubmitToolWidget::SubmitClicked()
 	}
 	else
 	{
-		if (ModelInterface->IsIntegrationRequired() && !ModelInterface->bIsUserInAllowlist)
+		if (ModelInterface->IsIntegrationRequired() && !ModelInterface->LockdownStatus.bIsAllowlisted)
 		{
 			IntegrationWidget->Open();
 			return FReply::Handled();
@@ -884,7 +884,7 @@ FText SubmitToolWidget::GetMainButtonText() const
 	}
 	else
 	{
-		FinalText = ModelInterface->IsIntegrationRequired() && !ModelInterface->bIsUserInAllowlist ? NSLOCTEXT("SourceControl.SubmitPanel", "IntegrationButton", "Open Integration Window") : NSLOCTEXT("SourceControl.SubmitPanel", "SubmitButton", "Submit");
+		FinalText = ModelInterface->IsIntegrationRequired() && !ModelInterface->LockdownStatus.bIsAllowlisted ? NSLOCTEXT("SourceControl.SubmitPanel", "IntegrationButton", "Open Integration Window") : NSLOCTEXT("SourceControl.SubmitPanel", "SubmitButton", "Submit");
 	}
 
 	return FinalText;
