Correcting when controlled lockdown is evaluated to properly communicate the messages to the user. Better naming for functions

#Synced-CL 50661505

diff --git a/Engine/Source/Programs/SubmitTool/Private/Models/ModelInterface.cpp b/Engine/Source/Programs/SubmitTool/Private/Models/ModelInterface.cpp
--- a/Engine/Source/Programs/SubmitTool/Private/Models/ModelInterface.cpp
+++ b/Engine/Source/Programs/SubmitTool/Private/Models/ModelInterface.cpp
@@ -104,7 +104,11 @@ FModelInterface::FModelInterface(FSubmitToolParameters&& InParameters, const TSh
 				}
 			}
 
-			UE_LOG(LogSubmitTool, Log, TEXT("The required local validation has succeeded, you're ALLOWED TO SUBMIT."))
+			if (EvaluateControlledLockdownAllowed())
+			{
+				UE_LOG(LogSubmitTool, Log, TEXT("The required local validation has succeeded, you're ALLOWED TO SUBMIT."))
+			}
+
 			if (ValidationService->GetIsAnyTaskRunning())
 			{
 				UE_LOG(LogSubmitTool, Warning, TEXT("You still have optional validations running you might want to consider waiting for them to finish."))
@@ -341,7 +345,6 @@ void FModelInterface::SetCLDescription(const FText& newDescription, bool DoNotIn
 	if (ChangelistService->SetCLDescription(newDescription.ToString()))
 	{
 		TagService->ParseCLDescription();
-		CheckControlledLockdownRequirements();
 
 		if (!DoNotInvalidate)
 		{
@@ -594,7 +597,7 @@ void FModelInterface::StartSubmitProcess(bool bSkipShelfDialog)
 			PrepareSubmitCallBack.Broadcast();
 		}
 
-		if (!CheckControlledLockdownRequirements())
+		if (!EvaluateControlledLockdownAllowed())
 		{
 			return;
 		}
@@ -1415,14 +1418,14 @@ void FModelInterface::ReevaluateControlledLockdown()
 			{
 				LockdownStatus = MoveTemp(LockdownData);
 				TitleMessage = FSubmitToolCoreUtils::StringBuilderJoin(LockdownStatus.GroupMessages, TEXT("\n"));
-				CheckControlledLockdownRequirements();
+				EvaluateControlledLockdownAllowed();
 			}, LowLevelTasks::ETaskPriority::Normal, UE::Tasks::EExtendedTaskPriority::GameThreadNormalPri);
 		});
 }
 
-bool FModelInterface::CheckControlledLockdownRequirements()
+bool FModelInterface::EvaluateControlledLockdownAllowed()
 {
-	TRACE_CPUPROFILER_EVENT_SCOPE(FModelInterface::CheckControlledLockdownRequirements);
+	TRACE_CPUPROFILER_EVENT_SCOPE(FModelInterface::EvaluateControlledLockdownAllowed);
 	if (LockdownStatus.LockdownType == ELockdownType::Controlled && !LockdownStatus.bIsAllowlisted)
 	{
 		bool bFillsRequirements = true;
@@ -1444,7 +1447,7 @@ bool FModelInterface::CheckControlledLockdownRequirements()
 				}
 			}
 
-			if (!bSubmitControlledLock && !bFillsTagRequirements)
+			if (!bFillsTagRequirements)
 			{
 				if (RequiredTag.Value.Num() == 0)
 				{
@@ -1461,21 +1464,21 @@ bool FModelInterface::CheckControlledLockdownRequirements()
 
 		if (bFillsRequirements || LockdownStatus.RequiredTags.Num() == 0)
 		{
-			bSubmitControlledLock = false;
-			return !bSubmitControlledLock;
+			bIsControlledLock = false;
+			return !bIsControlledLock;
 		}
 		else
 		{
-			if (!bSubmitControlledLock)
+			if (!bIsControlledLock)
 			{
 				UE_LOG(LogSubmitTool, Error, TEXT("Additional info:\n%s"), *FSubmitToolCoreUtils::StringBuilderJoin(LockdownStatus.GroupMessages, TEXT("\n")));
 			}
 
-			bSubmitControlledLock = true;
-			return !bSubmitControlledLock;
+			bIsControlledLock = true;
+			return !bIsControlledLock;
 		}
 	}
 
-	bSubmitControlledLock = false;
-	return !bSubmitControlledLock;
+	bIsControlledLock = false;
+	return !bIsControlledLock;
 }
diff --git a/Engine/Source/Programs/SubmitTool/Private/Models/ModelInterface.h b/Engine/Source/Programs/SubmitTool/Private/Models/ModelInterface.h
--- a/Engine/Source/Programs/SubmitTool/Private/Models/ModelInterface.h
+++ b/Engine/Source/Programs/SubmitTool/Private/Models/ModelInterface.h
@@ -193,13 +193,13 @@ public:
 	TMap<FName, TSet<FString>> DynamicValidatorLockdownIds;
 	void DynamicValidatorLockdownRefreshed(const FValidatorBase& InValidator);
 	void ReevaluateControlledLockdown();
-	bool CheckControlledLockdownRequirements();
+	bool EvaluateControlledLockdownAllowed();
 
 	const FSubmitToolParameters& GetParameters() const { return Parameters; }
 
 	const FString& GetTitleMessage() const;
 	bool IsIntegrationRequired() const { return LockdownStatus.LockdownType == ELockdownType::Hardcore; };
-	bool IsSubmissionBlocked() const { return bSubmitControlledLock; }
+	bool IsSubmissionBlocked() const { return bIsControlledLock; }
 	void RequestIntegration();
 	const TMap<FString, TSharedPtr<FIntegrationOptionBase>>& GetIntegrationOptions() { return IntegrationService->GetIntegrationOptions(); }
 	bool ValidateIntegrationOptions(bool bSilent) const	{ return IntegrationService->ValidateIntegrationOptions(bSilent); }
@@ -274,7 +274,7 @@ private:
 	TSharedPtr<FUpdateService> UpdateService;
 	TSharedPtr<FSubmitToolServiceProvider> ServiceProvider;
 	bool bPreflightQueued = false;
-	bool bSubmitControlledLock = false;
+	bool bIsControlledLock = false;
 	mutable FString TitleMessage;
 	bool Tick(float InDeltaTime);
 
