Fixing Preflight validator being case sensitive with 'none' and 'skip'

#Synced-CL 50661505

diff --git a/Engine/Source/Programs/SubmitTool/Private/Logic/Validators/PreflightValidator.cpp b/Engine/Source/Programs/SubmitTool/Private/Logic/Validators/PreflightValidator.cpp
--- a/Engine/Source/Programs/SubmitTool/Private/Logic/Validators/PreflightValidator.cpp
+++ b/Engine/Source/Programs/SubmitTool/Private/Logic/Validators/PreflightValidator.cpp
@@ -54,7 +54,7 @@ bool FPreflightValidator::Validate(const FString& InCLDescription, const TArray<
 	{
 		for(const FString& PreflightId : PreflightTag->GetValues())
 		{
-			if(PreflightId.Equals(TEXT("skip")) || PreflightId.Equals(TEXT("none")))
+			if(PreflightId.Equals(TEXT("skip"), ESearchCase::IgnoreCase) || PreflightId.Equals(TEXT("none"), ESearchCase::IgnoreCase))
 			{
 				// if we have at least one skip or none value, skip the validation altogether
 				Skip();
