Add changelist description generation to SubmitTool via external process

Implements a new "Generate" button in the SubmitTool UI that invokes an external process to automatically generate changelist descriptions.
Adds configuration parameters for the generation process, script path, parsing logic, and documentation URL.
Updates Tag model to support move semantics and bulk value updates.
The feature includes platform-specific script support (Windows .bat, Linux/Mac .sh) and displays a throbber during generation.

#Synced-CL 50661505

diff --git a/Engine/Source/Developer/SubmitToolCore/Private/Models/Tag.cpp b/Engine/Source/Developer/SubmitToolCore/Private/Models/Tag.cpp
--- a/Engine/Source/Developer/SubmitToolCore/Private/Models/Tag.cpp
+++ b/Engine/Source/Developer/SubmitToolCore/Private/Models/Tag.cpp
@@ -141,6 +141,35 @@ void FTag::SetValues(const TArray<FString>& InValues)
 	}
 }
 
+void FTag::SetValues(TArray<FString>&& InValues)
+{
+	bIsDirty = true;
+	this->TagValues = MoveTemp(InValues);
+
+	for(size_t i = 0; i<TagValues.Num();++i)
+	{
+		bool removedChar;
+		do
+		{
+			removedChar = false;
+			for(const FString& Delim : SplittedDelims)
+			{
+				bool localRemoved;
+				do
+				{
+					TagValues[i].TrimCharInline(Delim[0], &localRemoved);
+					removedChar |= localRemoved;
+				} while(localRemoved);
+			}
+		} while(removedChar);
+	}
+
+	if(OnTagUpdated.IsBound())
+	{
+		OnTagUpdated.Broadcast(*this);
+	}
+}
+
 const TArray<FString> FTag::GetValues(bool bEvenIfDisabled) const
 {
 	if(IsEnabled() || bEvenIfDisabled)
diff --git a/Engine/Source/Developer/SubmitToolCore/Public/Models/Tag.h b/Engine/Source/Developer/SubmitToolCore/Public/Models/Tag.h
--- a/Engine/Source/Developer/SubmitToolCore/Public/Models/Tag.h
+++ b/Engine/Source/Developer/SubmitToolCore/Public/Models/Tag.h
@@ -24,7 +24,27 @@ public:
 	FTag(const FTagDefinition& def, size_t start);
 
 	FTag(const FTag& InOther) : FTag(InOther.Definition, InOther.StartPos)
-	{}
+	{
+		TagValues = InOther.TagValues;
+		TagKeyValue = InOther.TagKeyValue;
+
+		ValidationState = ETagState::Unchecked;
+		bIsDirty = InOther.bIsDirty;
+		OnTagUpdated = InOther.OnTagUpdated;
+	}
+
+	FTag(FTag && InOther) : FTag(InOther.Definition, InOther.StartPos)
+	{
+		TagValues = MoveTemp(InOther.TagValues);
+		TagKeyValue = MoveTemp(InOther.TagKeyValue);
+
+		ValidationState = ETagState::Unchecked;
+		bIsDirty = MoveTemp(InOther.bIsDirty);
+		OnTagUpdated = MoveTemp(InOther.OnTagUpdated);
+	}
+
+	FTag& operator=(const FTag& InOther) = delete;
+	FTag& operator=(FTag&& InOther) = delete;
 
 	FString GetFullTag() const;
 	bool ParseTag(const FString& source);
@@ -33,6 +53,7 @@ public:
 	FString GetValuesText() const;
 
 	void SetValues(const TArray<FString>& InValues); 
+	void SetValues(TArray<FString>&& InValues );
 	const TArray<FString> GetValues(bool bIncludeInactive = false) const;
 
 	const FTagValidationConfig& GetCurrentValidationConfig(const TArray<FSCFileRef>& InDepotPaths) const;
diff --git a/Engine/Source/Developer/SubmitToolCore/Public/Parameters/SubmitToolParameters.h b/Engine/Source/Developer/SubmitToolCore/Public/Parameters/SubmitToolParameters.h
--- a/Engine/Source/Developer/SubmitToolCore/Public/Parameters/SubmitToolParameters.h
+++ b/Engine/Source/Developer/SubmitToolCore/Public/Parameters/SubmitToolParameters.h
@@ -69,6 +69,21 @@ struct FGeneralParameters
 
 	UPROPERTY()
 	TArray<FString> SyncOnStartupPaths;
+
+	UPROPERTY()
+	FString DescriptionGenProcess;
+
+	UPROPERTY()
+	FString DescriptionGenTooltip;
+
+	UPROPERTY()
+	FString DescriptionGenArgs;
+
+	UPROPERTY()
+	FString DescriptionGenParsing;
+
+	UPROPERTY()
+	FString DescriptionGenDocumentationUrl;
 };
 
 USTRUCT()
diff --git a/Engine/Source/Programs/SubmitTool/Private/Logic/TagService.cpp b/Engine/Source/Programs/SubmitTool/Private/Logic/TagService.cpp
--- a/Engine/Source/Programs/SubmitTool/Private/Logic/TagService.cpp
+++ b/Engine/Source/Programs/SubmitTool/Private/Logic/TagService.cpp
@@ -35,6 +35,30 @@ void FTagService::ParseCLDescription()
 	}
 }
 
+void FTagService::ReApplyTags()
+{
+	TArray<FTag*> TagsToReApply;
+
+	for (TPair<FString, FTag>& TagPair : RegisteredTags)
+	{
+		if (TagPair.Value.StartPos != std::numeric_limits<size_t>::max() && !GetCLDescription().Contains(TagPair.Value.Definition.TagId))
+		{
+			TagPair.Value.StartPos = std::numeric_limits<size_t>::max();
+			TagPair.Value.LastSize = std::numeric_limits<size_t>::max();
+			TagPair.Value.SetTagState(ETagState::Unchecked);
+			TagsToReApply.Add(&TagPair.Value);
+		}
+	}
+
+	for (FTag* Tag : TagsToReApply)
+	{
+		ApplyTag(*Tag);
+	}
+
+	// Parse to ensure we're matching the current description
+	ParseCLDescription();
+}
+
 void FTagService::ApplyTag(const FString& tagID)
 {
 	FTag* tag = GetTag(tagID);
diff --git a/Engine/Source/Programs/SubmitTool/Private/Logic/TagService.h b/Engine/Source/Programs/SubmitTool/Private/Logic/TagService.h
--- a/Engine/Source/Programs/SubmitTool/Private/Logic/TagService.h
+++ b/Engine/Source/Programs/SubmitTool/Private/Logic/TagService.h
@@ -14,8 +14,8 @@ public:
 	FTagService(const FSubmitToolParameters& InParameters, TSharedPtr<FChangelistService> clService);
 
 	void ParseCLDescription();
+	void ReApplyTags();
 
-	void RegisterTransientTag(FTagDefinition Definition);
 	void ApplyTag(const FString& tagID);
 	void ApplyTag(FTag& tag);
 	void UpdateTagsInCL();
diff --git a/Engine/Source/Programs/SubmitTool/Private/Models/ModelInterface.cpp b/Engine/Source/Programs/SubmitTool/Private/Models/ModelInterface.cpp
--- a/Engine/Source/Programs/SubmitTool/Private/Models/ModelInterface.cpp
+++ b/Engine/Source/Programs/SubmitTool/Private/Models/ModelInterface.cpp
@@ -18,6 +18,8 @@
 #include "Misc/StringOutputDevice.h"
 #include "Telemetry/TelemetryService.h"
 #include "ProfilingDebugging/CpuProfilerTrace.h"
+#include "Serialization/JsonReader.h"
+#include "Serialization/JsonSerializer.h"
 
 #include "CommandLine/CmdLineParameters.h"
 #include "Logic/Validators/CoreRedirectsValidator.h"
@@ -195,6 +197,11 @@ void FModelInterface::Dispose() const
 {
 	ChangelistService->CancelP4Operations();
 
+	if (DescGenProcess != nullptr && DescGenProcess->IsRunning())
+	{
+		DescGenProcess->Stop();
+	}
+
 	if (GetInputEnabled())
 	{
 		ChangelistService->UpdateP4CLDescriptionSynchronously();
@@ -338,7 +345,68 @@ void FModelInterface::ParsePreSubmitOperations() const
 	PresubmitOperationsService->InitializeTasks(Tasks);
 }
 
+void FModelInterface::GenerateCLDescription()
+{
+	if (!Parameters.GeneralParameters.DescriptionGenProcess.IsEmpty() && (DescGenProcess == nullptr || !DescGenProcess->IsRunning()))
+	{
+		GeneratedDescription.Empty();
+
+		FOnOutputLine LineHandler = FOnOutputLine::CreateLambda([this](const FString& InOutput, const EProcessOutputType& InType)
+		{
+			if (InType == EProcessOutputType::ProcessError)
+			{
+				UE_LOG(LogSubmitTool, Error, TEXT("[Description Generation]: %s"), *InOutput);
+			}
+			else
+			{
+
+				if (InOutput.StartsWith(Parameters.GeneralParameters.DescriptionGenParsing))
+				{
+					UE_LOG(LogSubmitToolDebug, Log, TEXT("[Description Generation]: %s"), *InOutput);
+					FString Content = InOutput.RightChop(Parameters.GeneralParameters.DescriptionGenParsing.Len());
+
+					// Try to parse as JSON for multi-line support
+					TSharedPtr<FJsonObject> JsonObject;
+					TSharedRef<TJsonReader<>> Reader = TJsonReaderFactory<>::Create(Content);
+					if (FJsonSerializer::Deserialize(Reader, JsonObject) && JsonObject.IsValid())
+					{
+						FString DescriptionValue;
+						if (JsonObject->TryGetStringField(TEXT("description"), DescriptionValue))
+						{
+							GeneratedDescription = DescriptionValue;
+						}
+					}
+					else
+					{
+						// Fallback to plain text for backwards compatibility
+						GeneratedDescription = Content;
+					}
+				}
+				else
+				{
+					UE_LOG(LogSubmitTool, Log, TEXT("[Description Generation]: %s"), *InOutput);
+				}
+			}
+		});
+
+		FOnCompleted CompleteHandler = FOnCompleted::CreateLambda([this](const int& InExitCode)
+		{
+			if(InExitCode == 0 && !GeneratedDescription.IsEmpty())
+			{
+				ChangelistService->SetCLDescription(GeneratedDescription);
+				TagService->ReApplyTags();
+				ValidateCLDescription();
+			}
+			else
+			{
+				UE_LOG(LogSubmitTool, Error, TEXT("[Description Generation]: Process exited with code %d"), InExitCode);
+			}
+		});
 
+		DescGenProcess = MakeUnique<FProcessWrapper>(TEXT("Description Gen Proc"), FConfiguration::Substitute(Parameters.GeneralParameters.DescriptionGenProcess), FConfiguration::Substitute(Parameters.GeneralParameters.DescriptionGenArgs), CompleteHandler, LineHandler);
+		DescGenProcess->Start();
+	}
+}
 
 void FModelInterface::SetCLDescription(const FText& newDescription, bool DoNotInvalidate)
 {
diff --git a/Engine/Source/Programs/SubmitTool/Private/Models/ModelInterface.h b/Engine/Source/Programs/SubmitTool/Private/Models/ModelInterface.h
--- a/Engine/Source/Programs/SubmitTool/Private/Models/ModelInterface.h
+++ b/Engine/Source/Programs/SubmitTool/Private/Models/ModelInterface.h
@@ -99,6 +99,7 @@ public:
 	void RegisterTagUpdatedCallback(const FTagUpdated::FDelegate Callback)			{ TagService->OnTagUpdated.Add(Callback); }
 
 	void SetCLDescription(const FText& newDescription, bool DoNotInvalidate = false);
+	void GenerateCLDescription();
 	void SendDescriptionToP4() const;
 	void UpdateCLFromP4Async() const;
 	const FString& GetCLDescription() const											{ return ChangelistService->GetCLDescription(); }
@@ -131,6 +132,8 @@ public:
 	void CancelValidations(const FName& InValidatorId = FName(), bool InbAsFailed = false) const { ValidationService->StopTasks(InValidatorId, InbAsFailed); }
 	void CheckForFileEdits() const													{ ValidationService->CheckForLocalFileEdit(); }
 
+	bool IsDescriptionGenerationRunning() const										{ return DescGenProcess != nullptr && DescGenProcess->IsRunning(); }
+
 	FDelegateHandle AddSingleValidatorFinishedCallback
 		(const FOnSingleTaskFinished::FDelegate& Callback) const					{ return ValidationService->OnSingleTaskFinished.Add(Callback); }
 
@@ -258,6 +261,8 @@ private:
 	FDelegateHandle OnValidationFinishedHandle;
 	FDelegateHandle OnPresubmitFinishedHandle;
 	FDelegateHandle OnSingleValidationFinishedHandle;
+	TUniquePtr<FProcessWrapper> DescGenProcess;
+	FString GeneratedDescription;
 
 	TSharedPtr<ISTSourceControlService> SourceControlService;
 	TSharedPtr<FChangelistService> ChangelistService;
diff --git a/Engine/Source/Programs/SubmitTool/Private/Version/AppVersionDefines.h b/Engine/Source/Programs/SubmitTool/Private/Version/AppVersionDefines.h
--- a/Engine/Source/Programs/SubmitTool/Private/Version/AppVersionDefines.h
+++ b/Engine/Source/Programs/SubmitTool/Private/Version/AppVersionDefines.h
@@ -21,7 +21,7 @@
 
 #define SUBMIT_TOOL_VERSION_MAJOR 6
 #define SUBMIT_TOOL_VERSION_MINOR 6
-#define SUBMIT_TOOL_VERSION_PATCH 5
+#define SUBMIT_TOOL_VERSION_PATCH 6
 
 #ifndef SUBMIT_TOOL_CHANGELIST
 #define SUBMIT_TOOL_CHANGELIST_STRING TEXT("NoCL")
diff --git a/Engine/Source/Programs/SubmitTool/Private/View/SubmitToolStyle.cpp b/Engine/Source/Programs/SubmitTool/Private/View/SubmitToolStyle.cpp
--- a/Engine/Source/Programs/SubmitTool/Private/View/SubmitToolStyle.cpp
+++ b/Engine/Source/Programs/SubmitTool/Private/View/SubmitToolStyle.cpp
@@ -152,6 +152,7 @@ FSubmitToolStyle::FSubmitToolStyle() : FSlateStyleSet(StyleName)
 	BranchedColor = FLinearColor::FromSRGBColor(FColor::FromHex("#00E4A0"));
 	StatusCheckedOutColor = FLinearColor::FromSRGBColor(FColor::FromHex("#1FE44B"));
 
+	Set("AppIcon.Generate", new CORE_IMAGE_BRUSH_SVG("Starship/Common/AutomationTools", CoreStyleConstants::Icon16x16));
 	Set("RevisionControl.CheckedOut", new CORE_IMAGE_BRUSH_SVG("Starship/SourceControl/SCC_CheckedOut", CoreStyleConstants::Icon16x16, StatusCheckedOutColor));
 	Set("RevisionControl.OpenForAdd", new CORE_IMAGE_BRUSH_SVG("Starship/SourceControl/RC_MarkedForAdd", CoreStyleConstants::Icon16x16, FStyleColors::AccentBlue));
 	Set("RevisionControl.MarkedForDelete", new CORE_IMAGE_BRUSH_SVG("Starship/SourceControl/SCC_MarkedForDelete", CoreStyleConstants::Icon16x16, FStyleColors::Error));
diff --git a/Engine/Source/Programs/SubmitTool/Private/View/Widgets/SubmitToolWidget.cpp b/Engine/Source/Programs/SubmitTool/Private/View/Widgets/SubmitToolWidget.cpp
--- a/Engine/Source/Programs/SubmitTool/Private/View/Widgets/SubmitToolWidget.cpp
+++ b/Engine/Source/Programs/SubmitTool/Private/View/Widgets/SubmitToolWidget.cpp
@@ -195,7 +195,7 @@ void SubmitToolWidget::Construct(const FArguments& InArgs)
 		.OnTextCommitted_Lambda([this](const FText& Text, ETextCommit::Type CommitType) { if(CommitType != ETextCommit::OnEnter) { ModelInterface->ValidateCLDescription(); }})
 		.AutoWrapText(true)
 		.OnIsTypedCharValid(FOnIsTypedCharValid::CreateLambda([](const TCHAR) { return true; }))
-		.IsReadOnly_Lambda([]{ return !FModelInterface::GetInputEnabled();});
+		.IsReadOnly_Lambda([this]{ return !FModelInterface::GetInputEnabled() || ModelInterface->IsDescriptionGenerationRunning();});
 
 	FString PerforceClientName;
 	FCmdLineParameters::Get().GetValue(FSubmitToolCmdLine::P4Client, PerforceClientName);
@@ -228,12 +228,88 @@ void SubmitToolWidget::Construct(const FArguments& InArgs)
 					SNew(SHorizontalBox)
 					+ SHorizontalBox::Slot()
 					.HAlign(HAlign_Left)
+					.VAlign(VAlign_Center)
 					.AutoWidth()
 					[
 						SNew(STextBlock)
 						.Font(FAppStyle::Get().GetFontStyle("SmallFontBold"))
 						.Text(NSLOCTEXT("SourceControl.SubmitPanel", "ChangeListDesc", "Changelist Description"))
 					]
+					+ SHorizontalBox::Slot()
+					.AutoWidth()
+					[
+						SNew(SButton)
+						.Visibility_Lambda([this] { return ModelInterface->GetParameters().GeneralParameters.DescriptionGenDocumentationUrl.IsEmpty() ? EVisibility::Collapsed : EVisibility::Visible; })
+						.ButtonStyle(FAppStyle::Get(), "InvisibleButton")
+						.IsFocusable(false)
+						.ToolTipText(FText::FromString(ModelInterface->GetParameters().GeneralParameters.DescriptionGenTooltip))
+						.Cursor(EMouseCursor::Hand)
+						.OnClicked_Lambda([this]()
+							{
+								if (!ModelInterface->GetParameters().GeneralParameters.DescriptionGenDocumentationUrl.IsEmpty())
+								{
+									FPlatformProcess::LaunchURL(*ModelInterface->GetParameters().GeneralParameters.DescriptionGenDocumentationUrl, nullptr, nullptr);
+								}
+
+								return FReply::Handled();
+							})
+						[
+							SNew(SImage)
+								.Image(FSubmitToolStyle::Get().GetBrush("AppIcon.DocumentationHelp"))
+						]
+					]
+					+ SHorizontalBox::Slot()
+					.AutoWidth()
+					.Padding(5, 0, 0, 0)
+					[
+						SNew(SButton)
+						.ContentPadding(-3)
+						.Visibility_Lambda([this] { return ModelInterface->GetParameters().GeneralParameters.DescriptionGenProcess.IsEmpty() ? EVisibility::Collapsed : EVisibility::Visible; })
+						.IsEnabled_Lambda([this] { return !ModelInterface->IsDescriptionGenerationRunning(); })
+						.IsFocusable(true)
+						.ToolTipText(FText::FromString(TEXT("Generate Description")))
+						.OnClicked_Lambda([this]()
+						{
+							ModelInterface->GenerateCLDescription();
+							return FReply::Handled();
+						})
+						[
+							SNew(SHorizontalBox)
+							+SHorizontalBox::Slot()
+							.Padding(0,0,2,0)
+							.VAlign(VAlign_Center)
+							.HAlign(HAlign_Center)
+							.AutoWidth()
+							[
+								SNew(SImage)
+								.Image(FSubmitToolStyle::Get().GetBrush("AppIcon.Generate"))
+							]
+							+SHorizontalBox::Slot()
+							.VAlign(VAlign_Center)
+							.HAlign(HAlign_Center)
+							.AutoWidth()
+							[
+								SNew(STextBlock)
+								.Font(FAppStyle::Get().GetFontStyle("SmallFontBold"))
+								.Text(FText::FromString(TEXT("Generate")))
+							]
+						]
+					]
+					+ SHorizontalBox::Slot()
+					.AutoWidth()
+					.Padding(5, 0)
+					.VAlign(VAlign_Center)
+					[
+						SNew(SThrobber)
+							.NumPieces(3)
+							.Visibility_Lambda([this]()
+							{
+								return ModelInterface->IsDescriptionGenerationRunning()
+									? EVisibility::Visible
+									: EVisibility::Collapsed;
+							})
+					]
+
 					+ SHorizontalBox::Slot()
 					.FillWidth(1.f)
 					[
@@ -243,6 +319,7 @@ void SubmitToolWidget::Construct(const FArguments& InArgs)
 					// STREAM 
 					+ SHorizontalBox::Slot()
 					.HAlign(HAlign_Right)
+					.VAlign(VAlign_Center)
 					.AutoWidth()
 					[
 						SNew(STextBlock)
@@ -251,6 +328,7 @@ void SubmitToolWidget::Construct(const FArguments& InArgs)
 					]
 					+ SHorizontalBox::Slot()
 					.HAlign(HAlign_Right)
+					.VAlign(VAlign_Center)
 					.Padding(5, 0, 0, 0)
 					.AutoWidth()
 					[
@@ -262,6 +340,7 @@ void SubmitToolWidget::Construct(const FArguments& InArgs)
 					// WORKSPACE
 					+ SHorizontalBox::Slot()
 					.HAlign(HAlign_Right)
+					.VAlign(VAlign_Center)
 					.Padding(15, 0, 0, 0)
 					.AutoWidth()
 					[
@@ -271,6 +350,7 @@ void SubmitToolWidget::Construct(const FArguments& InArgs)
 					]
 					+ SHorizontalBox::Slot()
 					.HAlign(HAlign_Right)
+					.VAlign(VAlign_Center)
 					.Padding(5, 0, 0, 0)
 					.AutoWidth()
 					[
diff --git a/Engine/Source/Programs/SubmitTool/SubmitTool.Build.cs b/Engine/Source/Programs/SubmitTool/SubmitTool.Build.cs
--- a/Engine/Source/Programs/SubmitTool/SubmitTool.Build.cs
+++ b/Engine/Source/Programs/SubmitTool/SubmitTool.Build.cs
@@ -106,6 +106,7 @@ public class SubmitTool : ModuleRules
 		RuntimeDependencies.Add("$(EngineDir)/Content/Editor/Slate/Old/DashedBorder.png", StagedFileType.UFS); 
 		RuntimeDependencies.Add("$(EngineDir)/Content/Editor/Slate/Old/ToolTip_Background.png", StagedFileType.UFS);
 		RuntimeDependencies.Add("$(EngineDir)/Content/Editor/Slate/Checkerboard.png", StagedFileType.UFS);
+		RuntimeDependencies.Add("$(EngineDir)/Content/Slate/Starship/Common/AutomationTools.svg", StagedFileType.UFS);
 		RuntimeDependencies.Add("$(EngineDir)/Content/Slate/Starship/SourceControl/SCC_CheckedOut.svg", StagedFileType.UFS);
 		RuntimeDependencies.Add("$(EngineDir)/Content/Slate/Starship/SourceControl/RC_MarkedForAdd.svg", StagedFileType.UFS);
 		RuntimeDependencies.Add("$(EngineDir)/Content/Slate/Starship/SourceControl/SCC_MarkedForDelete.svg", StagedFileType.UFS);
