Add warning for DefaultEngine.ini modifications in SubmitTool

Adds a new file group warning that alerts users when submitting DefaultEngine.ini alongside other changes, prompting them to confirm the submission is intentional.
FileGroups matching logic to support wildcard (*) patterns.

#Synced-CL 50661505

diff --git a/Engine/Source/Programs/SubmitTool/Private/Models/ModelInterface.cpp b/Engine/Source/Programs/SubmitTool/Private/Models/ModelInterface.cpp
--- a/Engine/Source/Programs/SubmitTool/Private/Models/ModelInterface.cpp
+++ b/Engine/Source/Programs/SubmitTool/Private/Models/ModelInterface.cpp
@@ -764,7 +764,7 @@ void FModelInterface::StartSubmitProcess(bool bSkipShelfDialog)
 					for (size_t i = 0; i < FileGroup.FileGroups.Num(); ++i)
 					{
 						const FString ReplacedPath = FileGroup.FileGroups[i].Replace(TEXT("$(StreamRoot)"), *GetCurrentStream());
-						if (File->GetDepotPath().Contains(ReplacedPath, ESearchCase::IgnoreCase))
+						if (ReplacedPath.Equals(TEXT("*")) || File->GetDepotPath().Contains(ReplacedPath, ESearchCase::IgnoreCase))
 						{
 							if (!Indexes.Contains(i))
 							{
