Set bRequiresPlatformSDK = true for SubmitTool build configuration

This ensures the SubmitTool build requires platform SDK availability during compilation, aligning with standard build requirements for platform-dependent tooling.

#Synced-CL 50661505

diff --git a/Engine/Source/Programs/SubmitTool/SubmitTool.Build.cs b/Engine/Source/Programs/SubmitTool/SubmitTool.Build.cs
--- a/Engine/Source/Programs/SubmitTool/SubmitTool.Build.cs
+++ b/Engine/Source/Programs/SubmitTool/SubmitTool.Build.cs
@@ -6,7 +6,8 @@ public class SubmitTool : ModuleRules
 	public SubmitTool(ReadOnlyTargetRules Target) : base(Target)
 	{
 		NumIncludedBytesPerUnityCPPOverride = 64 * 1024;
-	
+		bRequiresPlatformSDK = true;
+
 		PublicIncludePathModuleNames.AddRange(
 			new string[] {
 				"Launch",
