Making sure newly installed submit tool versions are updated from Horde

#Synced-CL 50661505

diff --git a/Engine/Source/Programs/SubmitTool/Private/Logic/UpdateService.cpp b/Engine/Source/Programs/SubmitTool/Private/Logic/UpdateService.cpp
--- a/Engine/Source/Programs/SubmitTool/Private/Logic/UpdateService.cpp
+++ b/Engine/Source/Programs/SubmitTool/Private/Logic/UpdateService.cpp
@@ -51,11 +51,6 @@ bool FUpdateService::CheckForNewVersion()
 	{
 		return false;
 	}
-	else if(LocalVersion.IsEmpty())
-	{
-		LocalVersion = RemoteVersion;
-		SaveLocalVersionToFile();
-	}
 
 	return LocalVersion != RemoteVersion;
 }
