Add Members field to P4 lockdown configuration and use it for required tags instead of AllowListers.
Parse 'members' from config sections to populate the new Members set, which will be displayed as required tag participants when a lockdown is controlled but the user is not allowlisted.

#Synced-CL 51015927

diff --git a/Engine/Source/Programs/SubmitTool/Private/Logic/P4LockdownService.cpp b/Engine/Source/Programs/SubmitTool/Private/Logic/P4LockdownService.cpp
--- a/Engine/Source/Programs/SubmitTool/Private/Logic/P4LockdownService.cpp
+++ b/Engine/Source/Programs/SubmitTool/Private/Logic/P4LockdownService.cpp
@@ -194,7 +194,7 @@ FSubmitToolLockdownData FP4LockdownService::IsPathInLockdown(const FSCFileRef& I
 
 			if (bOverallLockdownResult.LockdownType == ELockdownType::Controlled && !bAllowlisted)
 			{
-				bOverallLockdownResult.RequiredTags.Add(data.ValidTag, data.AllowListers);
+				bOverallLockdownResult.RequiredTags.Add(data.ValidTag, data.Members);
 			}
 		}
 	}
@@ -365,9 +365,26 @@ void FP4LockdownService::ParseAllowListFiles(const TMap<FString, FString>& InFil
 				for (FString& AllowLister : Allowlisters)
 				{
 					AllowLister.ToLowerInline();
+					AllowLister.TrimStartAndEndInline();
 					data->AllowListers.Add(MoveTemp(AllowLister));
 				}
 
+
+				const FConfigValue* MemberList = ConfigSection.Find(TEXT("members"));
+				if (MemberList != nullptr)
+				{
+					TArray<FString> MemberValues;
+					MemberList->GetSavedValue().ParseIntoArray(MemberValues, TEXT(","), true);
+
+					for (FString& Member : MemberValues)
+					{
+						Member.ToLowerInline();
+						Member.TrimStartAndEndInline();
+						data->Members.Add(MoveTemp(Member));
+					}
+				}
+
+
 				TArray<FString> View;
 				ConfigSection.MultiFind(TEXT("view"), View, true);
 				for (FString& ViewLine : View)
diff --git a/Engine/Source/Programs/SubmitTool/Private/Logic/P4LockdownService.h b/Engine/Source/Programs/SubmitTool/Private/Logic/P4LockdownService.h
--- a/Engine/Source/Programs/SubmitTool/Private/Logic/P4LockdownService.h
+++ b/Engine/Source/Programs/SubmitTool/Private/Logic/P4LockdownService.h
@@ -103,6 +103,7 @@ private:
 		ELockdownType LockdownType = ELockdownType::None;
 		FString ValidTag;
 		TSet<FString, FCaseInsensitiveStringKeyFuncs> AllowListers;
+		TSet<FString, FCaseInsensitiveStringKeyFuncs> Members;
 		TArray<TPair<bool, FString>> Views;
 	};
 
