Fix SwarmService reviewer processing to correctly handle groups and improve control flow.
Move continue statement inside user block to prevent skipping group processing.

#Synced-CL 51015927

diff --git a/Engine/Source/Programs/SubmitTool/Private/Logic/SwarmService.cpp b/Engine/Source/Programs/SubmitTool/Private/Logic/SwarmService.cpp
--- a/Engine/Source/Programs/SubmitTool/Private/Logic/SwarmService.cpp
+++ b/Engine/Source/Programs/SubmitTool/Private/Logic/SwarmService.cpp
@@ -119,19 +119,17 @@ void FSwarmService::CreateReview(const TArray<FString>& InReviewers, const OnCre
 			if (FoundUser != nullptr)
 			{
 				ReviewersObject.Add(MakeShared<FJsonValueString>((*FoundUser)->Username));
+				continue;
 			}
-			continue;
 		}
 
 		if (Groups != nullptr)
 		{
-			TSharedPtr<FString>* FoundGroup = Groups->FindByPredicate([&value](TSharedPtr<FString>& InGroupData) { return *InGroupData == value; });
+			TSharedPtr<FString>* FoundGroup = Groups->FindByPredicate([value = value.TrimChar(TCHAR('@'))](TSharedPtr<FString>& InGroupData) { return *InGroupData == value; });
 			if (FoundGroup != nullptr)
 			{
 				TSharedPtr<FJsonObject> GroupObject = MakeShared<FJsonObject>();
-				TSharedPtr<FJsonObject> GroupDetails = MakeShared<FJsonObject>();
-				GroupDetails->SetStringField(TEXT("required"), TEXT("false"));
-				GroupObject->SetObjectField(*(*FoundGroup), GroupDetails);
+				GroupObject->SetStringField(TEXT("name"), *(*FoundGroup));
 				GroupsObject.Add(MakeShared<FJsonValueObject>(GroupObject));
 			}
 		}
