Simplify submit tool args obtention to use FCommandLine::Get() directly instead of manually reconstructing arguments.
The function now returns the full command line with only '-cl default' replacement applied when needed, eliminating redundant parameter parsing and string concatenation.

#Synced-CL 51015927

diff --git a/Engine/Source/Programs/SubmitTool/Private/Logic/UpdateService.cpp b/Engine/Source/Programs/SubmitTool/Private/Logic/UpdateService.cpp
--- a/Engine/Source/Programs/SubmitTool/Private/Logic/UpdateService.cpp
+++ b/Engine/Source/Programs/SubmitTool/Private/Logic/UpdateService.cpp
@@ -420,36 +420,14 @@ void FUpdateService::OnProcessDownloadRequestStream(void* InDataPtr, int64& InOu
 
 FString FUpdateService::GetSubmitToolArgs() const
 {
+	FString Commandline = FCommandLine::Get();
 
-	FString Changelist;
-	FCmdLineParameters::Get().GetValue(FSubmitToolCmdLine::P4ChangeList, Changelist);
-
-	// If it's not a valid changelist, try to get it through ChangelistService since we might have used 'default'
-	if (!Changelist.IsNumeric())
+	if (Commandline.Contains(TEXT("-cl default")))
 	{
-		Changelist = ServiceProvider.Pin()->GetService<FChangelistService>()->GetCLID();
+		Commandline.ReplaceInline(TEXT("-cl default"), *FString::Printf(TEXT("-cl %s"), *ServiceProvider.Pin()->GetService<FChangelistService>()->GetCLID()));
 	}
 
-	FString PerforceServerAndPort;
-	FCmdLineParameters::Get().GetValue(FSubmitToolCmdLine::P4Server, PerforceServerAndPort);
-
-	FString PerforceUserName;
-	FCmdLineParameters::Get().GetValue(FSubmitToolCmdLine::P4User, PerforceUserName);
-
-	FString PerforceClientName;
-	FCmdLineParameters::Get().GetValue(FSubmitToolCmdLine::P4Client, PerforceClientName);
-
-	FString FromEditor;
-	FCmdLineParameters::Get().GetValue(FSubmitToolCmdLine::EditorFlag, FromEditor);
-
-	FString ExecutableArguments = "-server $(server) -user $(user) -client $(client) -cl $(changelist) $(fromEditor)";
-	ExecutableArguments = ExecutableArguments.Replace(TEXT("$(server)"), *PerforceServerAndPort);
-	ExecutableArguments = ExecutableArguments.Replace(TEXT("$(client)"), *PerforceClientName);
-	ExecutableArguments = ExecutableArguments.Replace(TEXT("$(user)"), *PerforceUserName);
-	ExecutableArguments = ExecutableArguments.Replace(TEXT("$(changelist)"), *Changelist);
-	ExecutableArguments = ExecutableArguments.Replace(TEXT("$(fromEditor)"), *FromEditor);
-
-	return ExecutableArguments;
+	return Commandline;
 }
 
 const FString FUpdateService::GetReadableDownloadSize()
