Update changelist description generation with improved formatting guidelines
Add telemetry tracking for generated changelist description process
Bump SubmitTool version to 6.7.0

#Synced-CL 51015927

diff --git a/Engine/Source/Programs/SubmitTool/Private/Models/ModelInterface.cpp b/Engine/Source/Programs/SubmitTool/Private/Models/ModelInterface.cpp
--- a/Engine/Source/Programs/SubmitTool/Private/Models/ModelInterface.cpp
+++ b/Engine/Source/Programs/SubmitTool/Private/Models/ModelInterface.cpp
@@ -391,6 +391,12 @@ void FModelInterface::GenerateCLDescription()
 
 		FOnCompleted CompleteHandler = FOnCompleted::CreateLambda([this](const int& InExitCode)
 		{
+			FTelemetryService::Get()->CustomEvent(TEXT("SubmitTool.StandAlone.GenerateDescription.Complete"),
+				MakeAnalyticsEventAttributeArray(
+					TEXT("ExitCode"), InExitCode
+				)
+			);
+
 			if(InExitCode == 0 && !GeneratedDescription.IsEmpty())
 			{
 				ChangelistService->SetCLDescription(GeneratedDescription);
diff --git a/Engine/Source/Programs/SubmitTool/Private/Version/AppVersionDefines.h b/Engine/Source/Programs/SubmitTool/Private/Version/AppVersionDefines.h
--- a/Engine/Source/Programs/SubmitTool/Private/Version/AppVersionDefines.h
+++ b/Engine/Source/Programs/SubmitTool/Private/Version/AppVersionDefines.h
@@ -20,8 +20,8 @@
 #endif
 
 #define SUBMIT_TOOL_VERSION_MAJOR 6
-#define SUBMIT_TOOL_VERSION_MINOR 6
-#define SUBMIT_TOOL_VERSION_PATCH 6
+#define SUBMIT_TOOL_VERSION_MINOR 7
+#define SUBMIT_TOOL_VERSION_PATCH 0
 
 #ifndef SUBMIT_TOOL_CHANGELIST
 #define SUBMIT_TOOL_CHANGELIST_STRING TEXT("NoCL")
