Add optional silent error handling parameter to file download operations
Allow callers to suppress error reporting when downloading files from source control

#Synced-CL 51168036

diff --git a/Engine/Source/Programs/SubmitTool/Private/Logic/P4LockdownService.cpp b/Engine/Source/Programs/SubmitTool/Private/Logic/P4LockdownService.cpp
--- a/Engine/Source/Programs/SubmitTool/Private/Logic/P4LockdownService.cpp
+++ b/Engine/Source/Programs/SubmitTool/Private/Logic/P4LockdownService.cpp
@@ -39,7 +39,7 @@ void FP4LockdownService::DownloadAllowListData()
 			DepotPaths.Add(Pair.Value);
 		}
 
-		DownloadTask = ServiceProvider.Pin()->GetService<ISTSourceControlService>()->DownloadFiles(MoveTemp(DepotPaths), DownloadedFiles);
+		DownloadTask = ServiceProvider.Pin()->GetService<ISTSourceControlService>()->DownloadFiles(MoveTemp(DepotPaths), DownloadedFiles, true);
 	}
 }
 
diff --git a/Engine/Source/Programs/SubmitTool/Private/Logic/Services/Interfaces/ISTSourceControlService.h b/Engine/Source/Programs/SubmitTool/Private/Logic/Services/Interfaces/ISTSourceControlService.h
--- a/Engine/Source/Programs/SubmitTool/Private/Logic/Services/Interfaces/ISTSourceControlService.h
+++ b/Engine/Source/Programs/SubmitTool/Private/Logic/Services/Interfaces/ISTSourceControlService.h
@@ -129,7 +129,7 @@ public:
 	virtual void GetGroups(const FOnGroupsGet::FDelegate& Callback) = 0;
 	virtual void GetUsersAndGroups(const FOnUsersAndGroupsGet::FDelegate& Callback) = 0;
 
-	virtual UE::Tasks::TTask<bool> DownloadFiles(TArray<FString>&& FilePaths, TMap<FString, FSharedBuffer>& OutFileBuffers) = 0;
+	virtual UE::Tasks::TTask<bool> DownloadFiles(TArray<FString>&& FilePaths, TMap<FString, FSharedBuffer>& OutFileBuffers, bool bSilentErrors = false) = 0;
 	virtual const TArray<TSharedPtr<FUserData>>& GetRecentUsers() const = 0;
 	virtual void AddRecentUser(TSharedPtr<FUserData>& User) = 0;
 	virtual const TArray<TSharedPtr<FString>>& GetRecentGroups() const = 0;
diff --git a/Engine/Source/Programs/SubmitTool/Private/Logic/Services/SourceControl/SubmitToolPerforce.cpp b/Engine/Source/Programs/SubmitTool/Private/Logic/Services/SourceControl/SubmitToolPerforce.cpp
--- a/Engine/Source/Programs/SubmitTool/Private/Logic/Services/SourceControl/SubmitToolPerforce.cpp
+++ b/Engine/Source/Programs/SubmitTool/Private/Logic/Services/SourceControl/SubmitToolPerforce.cpp
@@ -701,11 +701,12 @@ const FString FSubmitToolPerforce::GetRecentGroupsFilepath() const
 	return FPaths::Combine(FSubmitToolCoreUtils::GetLocalAppDataPath(), TEXT("SubmitTool"), TEXT("recent_groups.dat"));
 }
 
-UE::Tasks::TTask<bool> FSubmitToolPerforce::DownloadFiles(TArray<FString>&& FilePaths, TMap<FString, FSharedBuffer>& OutFileBuffers)
+UE::Tasks::TTask<bool> FSubmitToolPerforce::DownloadFiles(TArray<FString>&& FilePaths, TMap<FString, FSharedBuffer>& OutFileBuffers, bool bSilentErrors)
 {
 	FilePaths.Insert(TEXT("-q"), 0);
 	TSharedRef<FSCCommand> Command = MakeShared<FSCCommand>(TEXT("print"), MoveTemp(FilePaths));
 	Command->bCollectData = true;
+	Command->bSilentErrors = bSilentErrors;
 	return UE::Tasks::Launch(UE_SOURCE_LOCATION, [this, Command, &OutFileBuffers]
 	{
 		if (RunP4Command(*Command, DefaultFlags | EP4ClientUserFlags::CollectData) || Command->Data.Num() != 0)
diff --git a/Engine/Source/Programs/SubmitTool/Private/Logic/Services/SourceControl/SubmitToolPerforce.h b/Engine/Source/Programs/SubmitTool/Private/Logic/Services/SourceControl/SubmitToolPerforce.h
--- a/Engine/Source/Programs/SubmitTool/Private/Logic/Services/SourceControl/SubmitToolPerforce.h
+++ b/Engine/Source/Programs/SubmitTool/Private/Logic/Services/SourceControl/SubmitToolPerforce.h
@@ -85,7 +85,7 @@ public:
 	virtual void GetGroups(const FOnGroupsGet::FDelegate& Callback) override;
 	virtual void GetUsersAndGroups(const FOnUsersAndGroupsGet::FDelegate& Callback) override;
 
-	virtual UE::Tasks::TTask<bool> DownloadFiles(TArray<FString>&& FilePaths, TMap<FString, FSharedBuffer>& OutFileBuffers) override;
+	virtual UE::Tasks::TTask<bool> DownloadFiles(TArray<FString>&& FilePaths, TMap<FString, FSharedBuffer>& OutFileBuffers, bool bSilentErrors = false) override;
 
 	virtual const TArray<TSharedPtr<FUserData>>& GetRecentUsers() const override;
 	virtual void AddRecentUser(TSharedPtr<FUserData>& User) override;
