Add stealth mode and editor context execution controls to validators
Introduce stealth mode flag to run validators without blocking submissions while capturing telemetry
Add flags to restrict validator execution based on editor or standalone launch context

#Synced-CL 51168036

diff --git a/Engine/Source/Developer/SubmitToolCore/Public/Logic/Validators/ValidatorDefinition.h b/Engine/Source/Developer/SubmitToolCore/Public/Logic/Validators/ValidatorDefinition.h
--- a/Engine/Source/Developer/SubmitToolCore/Public/Logic/Validators/ValidatorDefinition.h
+++ b/Engine/Source/Developer/SubmitToolCore/Public/Logic/Validators/ValidatorDefinition.h
@@ -49,6 +49,13 @@ struct FValidatorDefinition
 	UPROPERTY()
 	FString Type;
 
+	/// <summary>
+	/// This is a mode that will make the validator not report any errors or warnings to the program/user but still log them to telemetry
+	/// It doesn't matter if it fails, it will always report as succeed and never block submission.
+	/// </summary>
+	UPROPERTY()
+	bool bIsStealthMode = false;
+
 	///
 	///	If true, this validator will be disabled
 	/// 
@@ -170,6 +177,18 @@ struct FValidatorDefinition
 	UPROPERTY()
 	bool bSkipWhenCalledFromEditor = false;
 
+	/// <summary>
+	/// This validator will only be applicable when running from editor (-from-editor is present)
+	/// </summary>
+	UPROPERTY()
+	bool bRequiresSTLaunchedFromEditor = false;
+
+	/// <summary>
+	/// This validator will only be applicable when NOT running from editor (-from-editor is not present)
+	/// </summary>
+	UPROPERTY()
+	bool bRequiresSTLaunchedStandalone = false;
+
 	///
 	///	Skip is forbidden when this text if found in the CL description
 	/// 
diff --git a/Engine/Source/Programs/SubmitTool/Private/Logic/TasksService.cpp b/Engine/Source/Programs/SubmitTool/Private/Logic/TasksService.cpp
--- a/Engine/Source/Programs/SubmitTool/Private/Logic/TasksService.cpp
+++ b/Engine/Source/Programs/SubmitTool/Private/Logic/TasksService.cpp
@@ -126,7 +126,7 @@ bool FTasksService::QueueForExecution(const TSharedPtr<FValidatorBase>& InTask,
 		return false;
 	}
 
-	if ((!InTask->IsRelevantToCL() && !InbForceRun))
+	if (!InTask->IsRelevantToCL())
 	{
 		InTask->SetNotApplicable();
 		return false;
@@ -168,7 +168,14 @@ bool FTasksService::QueueForExecution(const TSharedPtr<FValidatorBase>& InTask,
 
 void FTasksService::OnTaskFinishedCallback(const FValidatorBase& InTask)
 {
-	FTelemetryService::Get()->CustomEvent(TelemetryBaseId + TEXT(".Finished"), InTask.GetTelemetryAttributes());
+	if (InTask.GetFailureReason() != EFailureReason::Cancelled)
+	{
+		FTelemetryService::Get()->CustomEvent(TelemetryBaseId + TEXT(".Finished"), InTask.GetTelemetryAttributes());
+	}
+	else
+	{
+		FTelemetryService::Get()->CustomEvent(TelemetryBaseId + TEXT(".Cancelled"), InTask.GetTelemetryAttributes());
+	}
 
 	if(OnSingleTaskFinished.IsBound())
 	{
diff --git a/Engine/Source/Programs/SubmitTool/Private/Logic/Validators/EditorCommandletValidator.cpp b/Engine/Source/Programs/SubmitTool/Private/Logic/Validators/EditorCommandletValidator.cpp
--- a/Engine/Source/Programs/SubmitTool/Private/Logic/Validators/EditorCommandletValidator.cpp
+++ b/Engine/Source/Programs/SubmitTool/Private/Logic/Validators/EditorCommandletValidator.cpp
@@ -44,6 +44,8 @@ bool FEditorCommandletValidator::Activate()
 	ModifiableDefinition->EditorRecordsFile = FConfiguration::Substitute(ModifiableDefinition->EditorRecordsFile);
 	ModifiableDefinition->bValidateExecutableExists = false;
 
+	BuiltRegexProcessing = MakeShared<FRegexPattern>(RegexMessageProcessing, ERegexPatternFlags::CaseInsensitive);
+
 	return FValidatorRunExecutable::Activate();
 }
 
@@ -87,6 +89,15 @@ FString FEditorCommandletValidator::MakeTemporaryUserDirectory() const
 	return FString();
 }
 
+void FEditorCommandletValidator::ProcessSingleMessageForTelemetry(FString& InOutMessage) const
+{
+	FValidatorRunExecutable::ProcessSingleMessageForTelemetry(InOutMessage);
+
+	// Regex used for clearing the editor log timestamp, looks for "]Log******:" and removes everything before it
+	FRegexMatcher Matcher = FRegexMatcher(*BuiltRegexProcessing, InOutMessage);
+	InOutMessage = Matcher.ReplaceFirst(TEXT(""));
+}
+
 void FEditorCommandletValidator::GetEditorsForPaths(const TArray<FSCFileRef>& InFilteredFilesInCL, TMap<FString, FEditorParameters>& OutProjectEditorParameters) const
 {
 	const FEditorCommandletValidatorDefinition* TypedDefinition = GetTypedDefinition<FEditorCommandletValidatorDefinition>();
diff --git a/Engine/Source/Programs/SubmitTool/Private/Logic/Validators/EditorCommandletValidator.h b/Engine/Source/Programs/SubmitTool/Private/Logic/Validators/EditorCommandletValidator.h
--- a/Engine/Source/Programs/SubmitTool/Private/Logic/Validators/EditorCommandletValidator.h
+++ b/Engine/Source/Programs/SubmitTool/Private/Logic/Validators/EditorCommandletValidator.h
@@ -21,6 +21,7 @@ public:
 	virtual bool Activate() override;
 
 protected:
+	virtual void ProcessSingleMessageForTelemetry(FString& InOutMessage) const override;
 	void GetEditorsForPaths(const TArray<FSCFileRef>& InFilteredFilesInCL, TMap<FString, FEditorParameters>& OutProjectEditorParameters) const;
 	void SortFilesByProjects(const TArray<FSCFileRef>& InFiles, TMap<FString, TArray<FSCFileRef>>& OutProjects, TMap<FString, TArray<FSCFileRef>>& OutSubProjects) const;
 	TMap<FString, TSet<FString>> GetEnableDisablePluginArguments(const TArray<FSCFileRef>& InFilesInCL) const;
@@ -31,4 +32,8 @@ protected:
 
 private:
 	FString MakeTemporaryUserDirectory() const;
+
+	// Regex used for clearing the editor log timestamp, looks for "]Log******:" and removes everything before it
+	FString RegexMessageProcessing = TEXT(".*](?=Log\\w+:)");
+	TSharedPtr<FRegexPattern> BuiltRegexProcessing;
 };
diff --git a/Engine/Source/Programs/SubmitTool/Private/Logic/Validators/ValidatorBase.cpp b/Engine/Source/Programs/SubmitTool/Private/Logic/Validators/ValidatorBase.cpp
--- a/Engine/Source/Programs/SubmitTool/Private/Logic/Validators/ValidatorBase.cpp
+++ b/Engine/Source/Programs/SubmitTool/Private/Logic/Validators/ValidatorBase.cpp
@@ -15,6 +15,7 @@
 #include "Logic/Services/Interfaces/ICacheDataService.h"
 #include "Logic/Services/Interfaces/ISTSourceControlService.h"
 #include "SubmitToolCoreUtils.h"
+#include "Telemetry/TelemetryService.h"
 
 FValidatorBase::FValidatorBase(const FName& InNameId, const FSubmitToolParameters& InParameters, TWeakPtr<FSubmitToolServiceProvider> InServiceProvider, const FString& InDefinition) :
 	ValidatorNameID(InNameId),
@@ -73,7 +74,7 @@ void FValidatorBase::StartValidation()
 		OnLockdownIdsChanged.Broadcast(*this);
 	}
 
-	FailureReason.Reset();
+	FailureReason = EFailureReason::None;
 	ErrorListCache.Reset();
 	WarningListCache.Reset();
 
@@ -107,7 +108,7 @@ void FValidatorBase::StartValidation()
 					LogFailure(ActivationError);
 				}
 
-				FailureReason = TEXT("Validator setup is invalid");
+				FailureReason = EFailureReason::InvalidSetup;
 				ValidationFinished(false);
 			}
 		}
@@ -159,7 +160,7 @@ void FValidatorBase::StartValidation()
 					{						
 						if(Definition->bAllowSkippingOnMemoryFail && Result == EDialogFactoryResult::SecondButton)
 						{
-							FailureReason = TEXT("Skipped: Out of Memory");
+							FailureReason = EFailureReason::SkippedOOM;
 							UE_LOG(LogSubmitTool, Log, TEXT("[%s] Skipped due to lack of memory.\n%s"), *GetValidatorName(), *FinalMessage);
 							Skip();
 							return;
@@ -171,7 +172,7 @@ void FValidatorBase::StartValidation()
 							PreflightService->RequestPreflight(Definition->PreflightTemplate);
 						}
 
-						FailureReason = TEXT("Out of Memory");
+						FailureReason = EFailureReason::OutOfMemory;
 						ValidationFinished(false);
 						return;
 					}
@@ -186,7 +187,11 @@ void FValidatorBase::StartValidation()
 
 			if (!Validate(ChangelistService->GetCLDescription(), FilteredFiles, TagService->GetTagsArray()))
 			{
-				FailureReason = TEXT("Validation Failed");
+				if(FailureReason == EFailureReason::None)
+				{
+					FailureReason = EFailureReason::UnspecifiedFailure;
+				}
+
 				ValidationFinished(false);
 			}
 		}
@@ -247,6 +252,23 @@ void FValidatorBase::Tick(float InDeltaTime)
 	}
 }
 
+void FValidatorBase::CancelValidation(const bool InbAsFailed)
+{
+	if (State == EValidationStates::Disabled)
+	{
+		return;
+	}
+
+	if (State == EValidationStates::Running)
+	{
+		StopInternalValidations();
+	}
+
+	FailureReason = EFailureReason::Cancelled;
+	State = InbAsFailed ? EValidationStates::Failed : EValidationStates::Not_Run;
+	SetValidationFinished(State);
+}
+
 bool FValidatorBase::CanStartTask() const
 {
 	bool bCanRun = true;
@@ -360,53 +382,145 @@ const FString FValidatorBase::GetStatusText() const
 const TArray<FAnalyticsEventAttribute> FValidatorBase::GetTelemetryAttributes() const
 {
 	const TSharedPtr<FChangelistService> ChangelistService = ServiceProvider.Pin()->GetService<FChangelistService>();
-	FString ChangelistNumber = ChangelistService->GetCLID();
 
+	FString StealthStateStr = TEXT("No Stealth Mode");
+
+	if (Definition->bIsStealthMode)
+	{
+		StealthStateStr = StaticEnum<EValidationStates>()->GetNameStringByValue(static_cast<int64>(StealthState)).Replace(TEXT("_"), TEXT(" "));
+	}
 
 	return MakeAnalyticsEventAttributeArray(
 		TEXT("ValidatorID"), *GetValidatorNameId().ToString(),
 		TEXT("ValidatorName"), *GetValidatorName(),
-		TEXT("Status"), GetHasPassed(),
+		TEXT("Status"), Definition->bIsStealthMode ? bStealthFailure : GetHasPassed(),
+		TEXT("ValidatorState"), StaticEnum<EValidationStates>()->GetNameStringByValue(static_cast<int64>(State)).Replace(TEXT("_"), TEXT(" ")),
+		TEXT("IsStealthValidator"), Definition->bIsStealthMode,
+		TEXT("StealthActualState"), MoveTemp(StealthStateStr),
 		TEXT("Runtime"), RunTime,
 		TEXT("Stream"), ServiceProvider.Pin()->GetService<ISTSourceControlService>()->GetCurrentStreamName(),
-		TEXT("PendingChangelist"), MoveTemp(ChangelistNumber),
-		TEXT("FailureReason"), FailureReason,
+		TEXT("PendingChangelist"), ChangelistService->GetCLID(),
+		TEXT("FailureReason"), StaticEnum<EFailureReason>()->GetNameStringByValue(static_cast<int64>(FailureReason)),
 		TEXT("FromEditor"), FCmdLineParameters::Get().Contains(FSubmitToolCmdLine::EditorFlag),
-		TEXT("Errors"), ErrorListCache,
-		TEXT("Warnings"), WarningListCache
+		TEXT("Errors"), ProcessMessagesForTelemetry(ErrorListCache),
+		TEXT("Warnings"), ProcessMessagesForTelemetry(WarningListCache)
 	);
 }
 
-void FValidatorBase::ValidationFinished(const bool bHasPassed)
+TArray<FString> FValidatorBase::ProcessMessagesForTelemetry(const TArray<FString>& InMessages) const
 {
-	if (bHasPassed)
+	TArray<FString> ProcessedMessages = InMessages;
+
+	for (FString& Message : ProcessedMessages)
 	{
-		UE_LOG(LogValidatorsResult, Log, TEXT("[%s] Task Succeeded! (%s)"), *GetValidatorName(), *FGenericPlatformTime::PrettyTime(RunTime));
+		ProcessSingleMessageForTelemetry(Message);
+	}
+
+	return ProcessedMessages;
+}
+
+void FValidatorBase::ProcessSingleMessageForTelemetry(FString& InOutMessage) const
+{
+	// Remove commas to not mess up with telemetry
+	InOutMessage.ReplaceInline(TEXT(","), TEXT("[comma]"));
+}
+
+TArray<FAnalyticsEventAttribute> FValidatorBase::GetSingleMessageTelemetryAttributes(const FString& InMessage) const
+{
+	return MakeAnalyticsEventAttributeArray(
+		TEXT("ValidatorID"), *GetValidatorNameId().ToString(),
+		TEXT("ValidatorName"), *GetValidatorName(),
+		TEXT("IsStealthValidator"), Definition->bIsStealthMode,
+		TEXT("Runtime"), RunTime,
+		TEXT("Stream"), ServiceProvider.Pin()->GetService<ISTSourceControlService>()->GetCurrentStreamName(),
+		TEXT("PendingChangelist"), ServiceProvider.Pin()->GetService<FChangelistService>()->GetCLID(),
+		TEXT("FromEditor"), FCmdLineParameters::Get().Contains(FSubmitToolCmdLine::EditorFlag),
+		TEXT("Message"), InMessage
+	);
+}
 
-		if (Definition->bUsesIncrementalCache)
+void FValidatorBase::ValidationFinished(const bool bPassed)
+{
+	if (!Definition->bIsStealthMode) // Only log results when not in stealth mode
+	{
+		if (bPassed)
 		{
-			ServiceProvider.Pin()->GetService<ICacheDataService>()->UpdateLastValidationForFiles(ServiceProvider.Pin()->GetService<FChangelistService>()->GetCLID(), GetValidatorNameId(), GetValidationConfigId(), FilteredFiles, FDateTime::UtcNow());
+			UE_LOG(LogValidatorsResult, Log, TEXT("[%s] Task Succeeded! (%s)"), *GetValidatorName(), *FGenericPlatformTime::PrettyTime(RunTime));
+		}
+		else if (Definition->IsRequired)
+		{
+			UE_LOG(LogValidatorsResult, Error, TEXT("[%s] Failed on Required Task!"), *GetValidatorName());
+		}
+		else
+		{
+			UE_LOG(LogValidatorsResult, Warning, TEXT("[%s] Failed on Optional Task!"), *GetValidatorName());
+		}
+
+
+		if (!bPassed)
+		{
+			for (const FString& ErrorMsg : Definition->AdditionalValidationErrorMessages)
+			{
+				LogFailure(FString::Printf(TEXT("[%s] %s"), *GetValidatorName(), *ErrorMsg));
+			}
 		}
 	}
-	else if (Definition->IsRequired)
+
+	if (bPassed && Definition->bUsesIncrementalCache)
 	{
-		UE_LOG(LogValidatorsResult, Error, TEXT("[%s] Failed on Required Task!"), *GetValidatorName());
+		ServiceProvider.Pin()->GetService<ICacheDataService>()->UpdateLastValidationForFiles(ServiceProvider.Pin()->GetService<FChangelistService>()->GetCLID(), GetValidatorNameId(), GetValidationConfigId(), FilteredFiles, FDateTime::UtcNow());
 	}
-	else
+
+	EValidationStates NewState = bPassed ? EValidationStates::Valid : EValidationStates::Failed;
+	SetValidationFinished(NewState);
+}
+
+void FValidatorBase::LogFailure(const FString& FormattedMsg) const
+{
+	if (Definition->IsRequired)
 	{
-		UE_LOG(LogValidatorsResult, Warning, TEXT("[%s] Failed on Optional Task!"), *GetValidatorName());
-	}
+		if (Definition->bIsStealthMode)
+		{
+			UE_LOG(LogValidators, Log, TEXT("%s"), *FormattedMsg);
+		}
+		else
+		{
+			UE_LOG(LogValidators, Error, TEXT("%s"), *FormattedMsg);
+			UE_LOG(LogValidatorsResult, Error, TEXT("%s"), *FormattedMsg);
+		}
+
+	
 
-	if (!bHasPassed)
+		FScopeLock Lock(&ErrorLogsMutex);
+		ErrorListCache.Add(FormattedMsg);
+
+		FString TelemetryMsg = FormattedMsg;
+		ProcessSingleMessageForTelemetry(TelemetryMsg);
+
+		FTelemetryService::Get()->CustomEvent(TEXT("SubmitTool.Task.Error"), GetSingleMessageTelemetryAttributes(TelemetryMsg));
+
+	}
+	else
 	{
-		for (const FString& ErrorMsg : Definition->AdditionalValidationErrorMessages)
+		if (Definition->bIsStealthMode)
 		{
-			LogFailure(FString::Printf(TEXT("[%s] %s"), *GetValidatorName(), *ErrorMsg));
+			UE_LOG(LogValidators, Log, TEXT("%s"), *FormattedMsg);
 		}
+		else
+		{
+			UE_LOG(LogValidators, Warning, TEXT("%s"), *FormattedMsg);
+			UE_LOG(LogValidatorsResult, Warning, TEXT("%s"), *FormattedMsg);
+		}
+
+		FScopeLock Lock(&ErrorLogsMutex);
+		WarningListCache.Add(FormattedMsg);
+
+		FString TelemetryMsg = FormattedMsg;
+		ProcessSingleMessageForTelemetry(TelemetryMsg);
+
+		FTelemetryService::Get()->CustomEvent(TEXT("SubmitTool.Task.Warning"), GetSingleMessageTelemetryAttributes(TelemetryMsg));
 	}
 
-	EValidationStates NewState = bHasPassed ? EValidationStates::Valid : EValidationStates::Failed;
-	SetValidationFinished(NewState);
 }
 
 void FValidatorBase::SetValidationFinished(EValidationStates NewState)
@@ -415,6 +529,17 @@ void FValidatorBase::SetValidationFinished(EValidationStates NewState)
 
 	check(NewState != EValidationStates::Running);
 	check(NewState != EValidationStates::Queued);
+
+	if (Definition->bIsStealthMode)
+	{
+		StealthState = NewState;
+		bStealthFailure = NewState == EValidationStates::Timeout || NewState == EValidationStates::Failed;
+		if (bStealthFailure)
+		{
+			NewState = EValidationStates::Valid;
+		}
+	}
+
 	State = NewState;
 
 	if (OnValidationFinished.IsBound())
@@ -517,7 +642,7 @@ void FValidatorBase::SetSelectedOption(const FString& InOptionName, const FStrin
 
 bool FValidatorBase::CanPrintErrors() const
 {
-	if (State == EValidationStates::Failed || State == EValidationStates::Timeout || State == EValidationStates::Valid)
+	if (!Definition->bIsStealthMode && (State == EValidationStates::Failed || State == EValidationStates::Timeout || State == EValidationStates::Valid))
 	{
 		return !ErrorListCache.IsEmpty() || !WarningListCache.IsEmpty();
 	}
@@ -583,6 +708,11 @@ bool FValidatorBase::AppliesToFile(const FSCFileRef InFile, bool InbAllowIncreme
 		return true;
 	}
 
+	if (Definition->bRequiresSTLaunchedFromEditor && !FCmdLineParameters::Get().Contains(FSubmitToolCmdLine::EditorFlag))
+	{
+		return false;
+	}
+
 	FString Filename = InFile->GetFilename();
 	FPaths::NormalizeFilename(Filename);
 	uint32 FileHash = GetTypeHash(InFile->GetFilename());
@@ -698,6 +828,15 @@ bool FValidatorBase::AppliesToCL(const FString& InCLDescription, const TArray<FS
 		return true;
 	}
 
+	// Check "-from-editor" cmdline arg
+	// Don't apply if the validator requires being launched from editor and the arg is not present
+	if ((Definition->bRequiresSTLaunchedFromEditor && !FCmdLineParameters::Get().Contains(FSubmitToolCmdLine::EditorFlag))
+	// Don't apply if the validator requires being launched from standalone and the arg is present
+		|| (Definition->bRequiresSTLaunchedStandalone && FCmdLineParameters::Get().Contains(FSubmitToolCmdLine::EditorFlag)))
+	{
+		return false;
+	}
+
 	for (const FString& Filepath : Definition->CheckFileExists)
 	{
 		if (!FPaths::FileExists(FConfiguration::Substitute(Filepath)))
diff --git a/Engine/Source/Programs/SubmitTool/Private/Logic/Validators/ValidatorBase.h b/Engine/Source/Programs/SubmitTool/Private/Logic/Validators/ValidatorBase.h
--- a/Engine/Source/Programs/SubmitTool/Private/Logic/Validators/ValidatorBase.h
+++ b/Engine/Source/Programs/SubmitTool/Private/Logic/Validators/ValidatorBase.h
@@ -49,6 +49,21 @@ enum class EValidationStates : uint8
 	Disabled
 };
 
+
+UENUM()
+enum class EFailureReason : uint8
+{
+	None,
+	InvalidSetup,
+	OutOfMemory,
+	SkippedOOM,
+	Cancelled,
+	InvalidExitCode,
+	Errors,
+	UnspecifiedFailure
+};
+
+
 template<class T>
 concept DerivedFromDefinition = std::is_base_of<FValidatorDefinition, T>::value;
 
@@ -124,28 +139,19 @@ public:
 		return RunTime;
 	}
 
-	const EValidationStates& GetState() const
+	EValidationStates GetState() const
 	{
 		return State;
 	}
 
-	virtual void Tick(float InDeltaTime);
-
-	virtual void CancelValidation(bool InbAsFailed = false)
+	EFailureReason GetFailureReason() const
 	{
-		if (State == EValidationStates::Disabled)
-		{
-			return;
-		}
+		return FailureReason;
+	}
 
-		if(State == EValidationStates::Running)
-		{
-			StopInternalValidations();
-		}
+	virtual void Tick(float InDeltaTime);
 
-		State = InbAsFailed ? EValidationStates::Failed : EValidationStates::Not_Run;
-		SetValidationFinished(State);
-	}
+	virtual void CancelValidation(const bool bInAsFailed = false);
 
 	virtual void SetQueued(bool InbForceRun = false)
 	{
@@ -230,7 +236,7 @@ public:
 protected:
 	FName ValidatorNameID;
 	FString ValidatorName;
-	FString FailureReason;
+	EFailureReason FailureReason;
 	FValidatorOptionsProvider OptionsProvider;
 	TWeakPtr<FSubmitToolServiceProvider> ServiceProvider;
 	TArray<FSCFileRef> FilteredFiles;
@@ -241,7 +247,8 @@ protected:
 	TSet<FString> LockdownIds;
 	bool bForceRun = false;
 	const FSubmitToolParameters& SubmitToolParameters;
-
+	bool bStealthFailure = false;
+	EValidationStates StealthState;
 	FDateTime Start;
 	float RunTime = 0;
 
@@ -258,28 +265,13 @@ protected:
 	{};
 	virtual void ValidationFinished(const bool bHasPassed);
 
-	virtual void LogFailure(const FString& FormattedMsg) const
-	{
-		if (Definition->IsRequired)
-		{
-			UE_LOG(LogValidators, Error, TEXT("%s"), *FormattedMsg);
-			UE_LOG(LogValidatorsResult, Error, TEXT("%s"), *FormattedMsg);
-
-			FScopeLock Lock(&ErrorLogsMutex);
-			ErrorListCache.Add(FormattedMsg);
-		}
-		else
-		{
-			UE_LOG(LogValidators, Warning, TEXT("%s"), *FormattedMsg);
-			UE_LOG(LogValidatorsResult, Warning, TEXT("%s"), *FormattedMsg);
-
-			FScopeLock Lock(&ErrorLogsMutex);
-			WarningListCache.Add(FormattedMsg);
-		}
-
-	}
+	virtual void LogFailure(const FString& FormattedMsg) const;
 
 	virtual void Skip();
+	virtual TArray<FString> ProcessMessagesForTelemetry(const TArray<FString>& InMessages) const;
+	virtual void ProcessSingleMessageForTelemetry(FString& InOutMessage) const;
+	virtual TArray<FAnalyticsEventAttribute> GetSingleMessageTelemetryAttributes(const FString& InMessage) const;
+
 
 	EValidationStates State = EValidationStates::Not_Run;
 
diff --git a/Engine/Source/Programs/SubmitTool/Private/Logic/Validators/ValidatorRunExecutable.cpp b/Engine/Source/Programs/SubmitTool/Private/Logic/Validators/ValidatorRunExecutable.cpp
--- a/Engine/Source/Programs/SubmitTool/Private/Logic/Validators/ValidatorRunExecutable.cpp
+++ b/Engine/Source/Programs/SubmitTool/Private/Logic/Validators/ValidatorRunExecutable.cpp
@@ -390,6 +390,15 @@ void FValidatorRunExecutable::Tick(float DeltaTime)
 				if (!Data.bSuccess.IsSet())
 				{
 					bool bSuccess = Def->AllowedExitCodes.Contains(Process.ExitCode) && (Def->bOnlyLookAtExitCode || Data.ErrorList.IsEmpty());
+					if (!Def->AllowedExitCodes.Contains(Process.ExitCode))
+					{
+						FailureReason = EFailureReason::InvalidExitCode;
+					}
+					else
+					{
+						FailureReason = EFailureReason::Errors;
+					}
+
 					Data.bSuccess.Emplace(bSuccess);
 					TRACE_END_REGION(*Process.GetProcessName());
 				}
