Move changelist description validation from window close callback to tag value update
Ensure validation occurs immediately after tag values are set rather than during window destruction
Remove validation logic from window destroy override in dialog widgets

#Synced-CL 51168036
#Synced-CL 51168036

diff --git a/Engine/Source/Programs/SubmitTool/Private/View/Widgets/SJiraWidget.cpp b/Engine/Source/Programs/SubmitTool/Private/View/Widgets/SJiraWidget.cpp
--- a/Engine/Source/Programs/SubmitTool/Private/View/Widgets/SJiraWidget.cpp
+++ b/Engine/Source/Programs/SubmitTool/Private/View/Widgets/SJiraWidget.cpp
@@ -87,7 +87,7 @@ void SJiraWidget::Open(const FTag* InTargetTag)
 			.MinWidth(850)
 			.MinHeight(400);
 
-		ParentWindow->SetRequestDestroyWindowOverride(FRequestDestroyWindowOverride::CreateLambda([this](const TSharedRef<SWindow>& Window) { if(ModelInterface != nullptr) { ModelInterface->ValidateCLDescription(); } Window->HideWindow(); }));
+		ParentWindow->SetRequestDestroyWindowOverride(FRequestDestroyWindowOverride::CreateLambda([this](const TSharedRef<SWindow>& Window) { Window->HideWindow(); }));
 		FSlateApplication::Get().AddWindowAsNativeChild(ParentWindow.ToSharedRef(), MainWindow.ToSharedRef(), false);
 
 		FDeprecateSlateVector2D NewPosition = MainWindow->GetPositionInScreen();
@@ -503,6 +503,7 @@ void SJiraWidget::ProcessIssueSelected(TSharedPtr<FJiraIssue> InIssue)
 		}
 
 		ModelInterface->SetTagValues(*TargetTag, CurrentValues);
+		ModelInterface->ValidateCLDescription();
 	}
 }
 
diff --git a/Engine/Source/Programs/SubmitTool/Private/View/Widgets/SelectSourceControlGroupWidget.cpp b/Engine/Source/Programs/SubmitTool/Private/View/Widgets/SelectSourceControlGroupWidget.cpp
--- a/Engine/Source/Programs/SubmitTool/Private/View/Widgets/SelectSourceControlGroupWidget.cpp
+++ b/Engine/Source/Programs/SubmitTool/Private/View/Widgets/SelectSourceControlGroupWidget.cpp
@@ -96,7 +96,7 @@ FReply SSelectSourceControlGroupWidget::OpenDialog()
 				]
 		];
 
-	DialogWindow->SetRequestDestroyWindowOverride(FRequestDestroyWindowOverride::CreateLambda([this](const TSharedRef<SWindow>& Window) { if(ModelInterface != nullptr) { ModelInterface->ValidateCLDescription(); } Window->HideWindow(); }));
+	DialogWindow->SetRequestDestroyWindowOverride(FRequestDestroyWindowOverride::CreateLambda([this](const TSharedRef<SWindow>& Window) { Window->HideWindow(); }));
 
 	TSharedRef<SHorizontalBox> SearchBox = SNew(SHorizontalBox)
 		+ SHorizontalBox::Slot()
@@ -364,6 +364,7 @@ void SSelectSourceControlGroupWidget::ProcessGroupSelected(TSharedPtr<FString> I
 	}
 
 	ModelInterface->SetTagValues(*TargetTag, CurrentValues);
+	ModelInterface->ValidateCLDescription();
 
 	ModelInterface->AddRecentGroup(InGroup);
 
diff --git a/Engine/Source/Programs/SubmitTool/Private/View/Widgets/SelectSourceControlUserWidget.cpp b/Engine/Source/Programs/SubmitTool/Private/View/Widgets/SelectSourceControlUserWidget.cpp
--- a/Engine/Source/Programs/SubmitTool/Private/View/Widgets/SelectSourceControlUserWidget.cpp
+++ b/Engine/Source/Programs/SubmitTool/Private/View/Widgets/SelectSourceControlUserWidget.cpp
@@ -101,7 +101,7 @@ FReply SSelectSourceControlUserWidget::OpenDialog()
 			]
 		];
 
-	DialogWindow->SetRequestDestroyWindowOverride(FRequestDestroyWindowOverride::CreateLambda([this](const TSharedRef<SWindow>& Window) { if(ModelInterface != nullptr) { ModelInterface->ValidateCLDescription(); } Window->HideWindow(); }));
+	DialogWindow->SetRequestDestroyWindowOverride(FRequestDestroyWindowOverride::CreateLambda([this](const TSharedRef<SWindow>& Window) { Window->HideWindow(); }));
 
 	TSharedRef<SHorizontalBox> SearchBox = SNew(SHorizontalBox)
 		+ SHorizontalBox::Slot()
@@ -464,6 +464,7 @@ void SSelectSourceControlUserWidget::ProcessUserSelected(TSharedPtr<FUserData> I
 	if(TargetTag)
 	{
 		ModelInterface->SetTagValues(*TargetTag, CurrentValues);
+		ModelInterface->ValidateCLDescription();
 	}
 	else
 	{
