Add telemetry exit event and prevent telemetry recording during application shutdown

#Synced-CL 51168036

diff --git a/Engine/Source/Programs/SubmitTool/Private/SubmitToolApp.cpp b/Engine/Source/Programs/SubmitTool/Private/SubmitToolApp.cpp
--- a/Engine/Source/Programs/SubmitTool/Private/SubmitToolApp.cpp
+++ b/Engine/Source/Programs/SubmitTool/Private/SubmitToolApp.cpp
@@ -157,6 +157,7 @@ static void EngineLoop(FModelInterface* ModelInterface)
 
 		if (IsEngineExitRequested())
 		{
+			FTelemetryService::Get()->Exit();
 			// Dispose here so SCCProvider has time to clean up
 			ModelInterface->Dispose();
 		}
diff --git a/Engine/Source/Programs/SubmitTool/Private/Telemetry/ITelemetry.h b/Engine/Source/Programs/SubmitTool/Private/Telemetry/ITelemetry.h
--- a/Engine/Source/Programs/SubmitTool/Private/Telemetry/ITelemetry.h
+++ b/Engine/Source/Programs/SubmitTool/Private/Telemetry/ITelemetry.h
@@ -13,4 +13,9 @@ public:
 	virtual void BlockFlush(float InTimeout) const = 0;
 	virtual void SubmitSucceeded(TArray<FAnalyticsEventAttribute>&& InAttribs) const = 0;
 	virtual void CustomEvent(const FString& InEventId, const TArray<FAnalyticsEventAttribute>& InAttribs) const = 0;
+	virtual void Exit() = 0;
+
+
+protected:
+	bool bExiting = false;
 };
\ No newline at end of file
diff --git a/Engine/Source/Programs/SubmitTool/Private/Telemetry/NullTelemetry.h b/Engine/Source/Programs/SubmitTool/Private/Telemetry/NullTelemetry.h
--- a/Engine/Source/Programs/SubmitTool/Private/Telemetry/NullTelemetry.h
+++ b/Engine/Source/Programs/SubmitTool/Private/Telemetry/NullTelemetry.h
@@ -11,4 +11,6 @@ public:
 	virtual void BlockFlush(float InTimeout) const {}
 	virtual void SubmitSucceeded(TArray<FAnalyticsEventAttribute>&& InAttribs) const override {}
 	virtual void CustomEvent(const FString& InEventId, const TArray<FAnalyticsEventAttribute>& InAttribs) const override {}
+	virtual void Exit() {}
+
 };
\ No newline at end of file
diff --git a/Engine/Source/Programs/SubmitTool/Private/Telemetry/StandAloneTelemetry.cpp b/Engine/Source/Programs/SubmitTool/Private/Telemetry/StandAloneTelemetry.cpp
--- a/Engine/Source/Programs/SubmitTool/Private/Telemetry/StandAloneTelemetry.cpp
+++ b/Engine/Source/Programs/SubmitTool/Private/Telemetry/StandAloneTelemetry.cpp
@@ -9,12 +9,13 @@
 FStandAloneTelemetry::FStandAloneTelemetry(const FString& InUrl, const FGuid& InSessionID)
 {
 	FAnalyticsET::Config Config;
-
-	// this is a huge bucket where we all !shipping builds are considered local dev builds
-#if UE_BUILD_SHIPPING || UE_BUILD_DEVELOPMENT
+	
+#if defined(SUBMIT_TOOL_PRERELEASE)
+	Config.APIKeyET = TEXT("SubmitToolStandalone.PreRelease");
+#elif defined(SUBMIT_TOOL_RELEASE)
 	Config.APIKeyET = TEXT("SubmitToolStandalone.Live");
 #else
-	Config.APIKeyET = TEXT("SubmitToolStandalone.Debug");
+	Config.APIKeyET = TEXT("SubmitToolStandalone.Source");
 #endif
 
 	Config.APIServerET = InUrl;
@@ -27,7 +28,7 @@ FStandAloneTelemetry::FStandAloneTelemetry(const FString& InUrl, const FGuid& In
 
 	// There are other things to configure, but the default are usually fine.
 	
-	this->Provider = FAnalyticsET::Get().CreateAnalyticsProvider(Config);
+	Provider = FAnalyticsET::Get().CreateAnalyticsProvider(Config);
 	checkf(Provider.IsValid(), TEXT("Failure constructing analytics provider!"));
 
 	Provider->SetUserID(InSessionID.ToString());
@@ -40,7 +41,7 @@ FStandAloneTelemetry::~FStandAloneTelemetry()
 
 void FStandAloneTelemetry::Start(const FString& InCurrentStream) const
 {
-	if (this->Provider == nullptr)
+	if (Provider == nullptr)
 	{
 		return;
 	}
@@ -65,7 +66,7 @@ void FStandAloneTelemetry::BlockFlush(float InTimeout) const
 }
 void FStandAloneTelemetry::CustomEvent(const FString& InEventId, const TArray<FAnalyticsEventAttribute>& InAttribs) const
 {
-	if(this->Provider == nullptr)
+	if(Provider == nullptr || bExiting)
 	{
 		return;
 	}
@@ -73,9 +74,26 @@ void FStandAloneTelemetry::CustomEvent(const FString& InEventId, const TArray<FA
 	Provider->RecordEvent(InEventId, InAttribs);
 }
 
+void FStandAloneTelemetry::Exit()
+{
+	bExiting = true;
+
+	if (Provider == nullptr)
+	{
+		return;
+	}
+
+	Provider->RecordEvent(
+		TEXT("SubmitTool.StandAlone.Exit"),
+		MakeAnalyticsEventAttributeArray(
+			TEXT("Version"), FAppVersion::GetVersion()
+		)
+	);
+}
+
 void FStandAloneTelemetry::SubmitSucceeded(TArray<FAnalyticsEventAttribute>&& InAttribs) const
 {
-	if (this->Provider == nullptr)
+	if (Provider == nullptr || bExiting)
 	{
 		return;
 	}
diff --git a/Engine/Source/Programs/SubmitTool/Private/Telemetry/StandAloneTelemetry.h b/Engine/Source/Programs/SubmitTool/Private/Telemetry/StandAloneTelemetry.h
--- a/Engine/Source/Programs/SubmitTool/Private/Telemetry/StandAloneTelemetry.h
+++ b/Engine/Source/Programs/SubmitTool/Private/Telemetry/StandAloneTelemetry.h
@@ -17,6 +17,8 @@ public:
 	virtual void BlockFlush(float InTimeout) const override;
 	virtual void SubmitSucceeded(TArray<FAnalyticsEventAttribute>&& InAttribs) const override;
 	virtual void CustomEvent(const FString& InEventId, const TArray<FAnalyticsEventAttribute>& InAttribs) const override;
+	virtual void Exit() override;
+
 
 private:
 	TSharedPtr<IAnalyticsProviderET> Provider;
diff --git a/Engine/Source/Programs/SubmitTool/Private/Telemetry/SubmitToolStudioTelemetry.cpp b/Engine/Source/Programs/SubmitTool/Private/Telemetry/SubmitToolStudioTelemetry.cpp
--- a/Engine/Source/Programs/SubmitTool/Private/Telemetry/SubmitToolStudioTelemetry.cpp
+++ b/Engine/Source/Programs/SubmitTool/Private/Telemetry/SubmitToolStudioTelemetry.cpp
@@ -37,16 +37,33 @@ void FSubmitToolStudioTelemetry::BlockFlush(float InTimeout) const
 }
 void FSubmitToolStudioTelemetry::CustomEvent(const FString& InEventId, const TArray<FAnalyticsEventAttribute>& InAttribs) const
 {
-	FStudioTelemetry::Get().RecordEvent(InEventId, InAttribs);
+	if (!bExiting)
+	{
+		FStudioTelemetry::Get().RecordEvent(InEventId, InAttribs);
+	}
 }
 
-void FSubmitToolStudioTelemetry::SubmitSucceeded(TArray<FAnalyticsEventAttribute>&& InAttribs) const
+void FSubmitToolStudioTelemetry::Exit()
 {
+	bExiting = true;
 	FStudioTelemetry::Get().RecordEvent(
-		TEXT("SubmitTool.StandAlone.Submit.Succeeded"),
-		AppendAnalyticsEventAttributeArray(
-			InAttribs,
+		TEXT("SubmitTool.StandAlone.Exit"),
+		MakeAnalyticsEventAttributeArray(
 			TEXT("Version"), FAppVersion::GetVersion()
 		)
 	);
+}
+
+void FSubmitToolStudioTelemetry::SubmitSucceeded(TArray<FAnalyticsEventAttribute>&& InAttribs) const
+{
+	if (!bExiting)
+	{ 
+		FStudioTelemetry::Get().RecordEvent(
+			TEXT("SubmitTool.StandAlone.Submit.Succeeded"),
+			AppendAnalyticsEventAttributeArray(
+				InAttribs,
+				TEXT("Version"), FAppVersion::GetVersion()
+			)
+		);
+	}
 }
\ No newline at end of file
diff --git a/Engine/Source/Programs/SubmitTool/Private/Telemetry/SubmitToolStudioTelemetry.h b/Engine/Source/Programs/SubmitTool/Private/Telemetry/SubmitToolStudioTelemetry.h
--- a/Engine/Source/Programs/SubmitTool/Private/Telemetry/SubmitToolStudioTelemetry.h
+++ b/Engine/Source/Programs/SubmitTool/Private/Telemetry/SubmitToolStudioTelemetry.h
@@ -13,4 +13,5 @@ public:
 	virtual void BlockFlush(float InTimeout) const override;
 	virtual void SubmitSucceeded(TArray<FAnalyticsEventAttribute>&& InAttribs) const override;
 	virtual void CustomEvent(const FString& InEventId, const TArray<FAnalyticsEventAttribute>& InAttribs) const override;
+	virtual void Exit() override;
 };
\ No newline at end of file
