Add window focus check before drawing attention
Validate changelist description when validating integration and the submit tool tag is missing

#Synced-CL 51168036

diff --git a/Engine/Source/Programs/SubmitTool/Private/Models/ModelInterface.cpp b/Engine/Source/Programs/SubmitTool/Private/Models/ModelInterface.cpp
--- a/Engine/Source/Programs/SubmitTool/Private/Models/ModelInterface.cpp
+++ b/Engine/Source/Programs/SubmitTool/Private/Models/ModelInterface.cpp
@@ -135,7 +135,10 @@ FModelInterface::FModelInterface(FSubmitToolParameters&& InParameters, const TSh
 
 	OnValidationFinishedHandle = ValidationService->OnTasksQueueFinished.Add(FOnTaskFinished::FDelegate::CreateLambda([this](bool bIsValid)
 		{
-			MainTab.Pin()->GetParentWindow()->DrawAttention(FWindowDrawAttentionParameters());
+			if (!MainTab.Pin()->GetParentWindow()->HasAnyUserFocusOrFocusedDescendants())
+			{
+				MainTab.Pin()->GetParentWindow()->DrawAttention(FWindowDrawAttentionParameters());
+			}
 
 			if (bIsValid)
 			{
diff --git a/Engine/Source/Programs/SubmitTool/Private/Models/ModelInterface.h b/Engine/Source/Programs/SubmitTool/Private/Models/ModelInterface.h
--- a/Engine/Source/Programs/SubmitTool/Private/Models/ModelInterface.h
+++ b/Engine/Source/Programs/SubmitTool/Private/Models/ModelInterface.h
@@ -205,7 +205,19 @@ public:
 	bool IsSubmissionBlocked() const { return bIsControlledLock; }
 	void RequestIntegration();
 	const TMap<FString, TSharedPtr<FIntegrationOptionBase>>& GetIntegrationOptions() { return IntegrationService->GetIntegrationOptions(); }
-	bool ValidateIntegrationOptions(bool bSilent) const	{ return IntegrationService->ValidateIntegrationOptions(bSilent); }
+	bool ValidateIntegrationOptions(bool bSilent) const	
+	{
+		if (IntegrationService->ValidateIntegrationOptions(bSilent))
+		{
+			if (!HasSubmitToolTag())
+			{
+				ValidateCLDescription();
+			}
+
+			return true;
+		}
+		return false;
+	}
 
 
 	static FOnStateChanged OnStateChanged;
