Add submitted changelist number to telemetry event data

#Synced-CL 51168036
#Synced-CL 51168036

diff --git a/Engine/Source/Programs/SubmitTool/Private/Models/ModelInterface.cpp b/Engine/Source/Programs/SubmitTool/Private/Models/ModelInterface.cpp
--- a/Engine/Source/Programs/SubmitTool/Private/Models/ModelInterface.cpp
+++ b/Engine/Source/Programs/SubmitTool/Private/Models/ModelInterface.cpp
@@ -1161,7 +1161,8 @@ void FModelInterface::OnSubmitOperationComplete(const FSCCommand& InCmd)
 		}
 	
 		FTelemetryService::Get()->SubmitSucceeded(MakeAnalyticsEventAttributeArray(
-			TEXT("Stream"), SourceControlService->GetCurrentStreamName()
+			TEXT("Stream"), SourceControlService->GetCurrentStreamName(),
+			TEXT("SubmittedCL"), SubmittedCLString
 		));
 
 		// We've submitted, or tried to submit and failed so we only let the user close the app
