[SubmitTool] Fix telemetry submission to dereference SubmittedCLString pointer when logging the submitted CL number. The string was being passed by pointer rather than by value to the analytics event.

#Synced-CL 51401191

#Synced-CL 51401191

diff --git a/Engine/Source/Programs/SubmitTool/Private/Models/ModelInterface.cpp b/Engine/Source/Programs/SubmitTool/Private/Models/ModelInterface.cpp
--- a/Engine/Source/Programs/SubmitTool/Private/Models/ModelInterface.cpp
+++ b/Engine/Source/Programs/SubmitTool/Private/Models/ModelInterface.cpp
@@ -1162,7 +1162,7 @@ void FModelInterface::OnSubmitOperationComplete(const FSCCommand& InCmd)
 	
 		FTelemetryService::Get()->SubmitSucceeded(MakeAnalyticsEventAttributeArray(
 			TEXT("Stream"), SourceControlService->GetCurrentStreamName(),
-			TEXT("SubmittedCL"), SubmittedCLString
+			TEXT("SubmittedCL"), ((SubmittedCLString == nullptr) ? FString() : *SubmittedCLString)
 		));
 
 		// We've submitted, or tried to submit and failed so we only let the user close the app
