Add a configurable regex stream allowlist to telemetry initialization
Fall back to null telemetry when the current stream does not match the allowlist pattern
Refactor telemetry parameter loading to use struct import via config section

#Synced-CL 51168036
#Synced-CL 51168036

diff --git a/Engine/Source/Developer/SubmitToolCore/Public/Parameters/SubmitToolParameters.h b/Engine/Source/Developer/SubmitToolCore/Public/Parameters/SubmitToolParameters.h
--- a/Engine/Source/Developer/SubmitToolCore/Public/Parameters/SubmitToolParameters.h
+++ b/Engine/Source/Developer/SubmitToolCore/Public/Parameters/SubmitToolParameters.h
@@ -167,6 +167,8 @@ struct FTelemetryParameters
 	FString Url;
 	UPROPERTY()
 	FString Instance;
+	UPROPERTY()
+	FString RegexStreamAllowlist;
 };
 
 UENUM()
diff --git a/Engine/Source/Programs/SubmitTool/Private/Parameters/SubmitToolParametersBuilder.cpp b/Engine/Source/Programs/SubmitTool/Private/Parameters/SubmitToolParametersBuilder.cpp
--- a/Engine/Source/Programs/SubmitTool/Private/Parameters/SubmitToolParametersBuilder.cpp
+++ b/Engine/Source/Programs/SubmitTool/Private/Parameters/SubmitToolParametersBuilder.cpp
@@ -196,25 +196,22 @@ FJiraParameters FSubmitToolParametersBuilder::BuildJiraParameters(const FConfigF
 
  FTelemetryParameters FSubmitToolParametersBuilder::GetTelemetryParameters(const FConfigFile& InConfigFile)
 {
-	static const TCHAR* Section = TEXT("SubmitTool.Telemetry");
-	static const TCHAR* UrlKey = TEXT("URL");
-	static const TCHAR* InstanceKey = TEXT("Instance");
+	 const FConfigSection* Section = InConfigFile.FindSection(TEXT("SubmitTool.Telemetry"));
+	 FTelemetryParameters Output;
 
-	FTelemetryParameters Output;
-
-	FString Url;
-	if (InConfigFile.GetString(Section, UrlKey, Url))
-	{
-		Output.Url = Url;
-	}
+	 if (Section != nullptr)
+	 {
+		 FStringOutputDevice Errors;
+		 FTelemetryParameters::StaticStruct()->ImportText(*SectionToText(*Section), &Output, nullptr, 0, &Errors, FTelemetryParameters::StaticStruct()->GetName());
 
-	FString Instance;
-	if (InConfigFile.GetString(Section, InstanceKey, Instance))
-	{
-		Output.Instance = Instance;
-	}
+		 if (!Errors.IsEmpty())
+		 {
+			 UE_LOG(LogSubmitTool, Error, TEXT("Error loading parameter file %s"), *Errors);
+			 FModelInterface::SetErrorState();
+		 }
+	 }
 
-	return Output;
+	 return Output;
 }
 
  FIntegrationParameters FSubmitToolParametersBuilder::BuildIntegrationParameters(const FConfigFile& InConfigFile)
diff --git a/Engine/Source/Programs/SubmitTool/Private/SubmitToolApp.cpp b/Engine/Source/Programs/SubmitTool/Private/SubmitToolApp.cpp
--- a/Engine/Source/Programs/SubmitTool/Private/SubmitToolApp.cpp
+++ b/Engine/Source/Programs/SubmitTool/Private/SubmitToolApp.cpp
@@ -278,8 +278,10 @@ int RunSubmitTool(const TCHAR* CommandLine)
 
 		// record that the application has started
 		// do this after initializing UI to give ISTSourceControlService time to populate GetCurrentStreamName
-		FTelemetryService::Init(ModelInterface->GetParameters().Telemetry);
-		FTelemetryService::Get()->Start(ModelInterface->GetServiceProvider()->GetService<ISTSourceControlService>()->GetCurrentStreamName());
+
+		const FString StreamName = ModelInterface->GetServiceProvider()->GetService<ISTSourceControlService>()->GetCurrentStreamName();
+		FTelemetryService::Init(ModelInterface->GetParameters().Telemetry, StreamName);
+		FTelemetryService::Get()->Start(StreamName);
 
 		// Tick the application until shutdown
 		EngineLoop(ModelInterface.Get());
diff --git a/Engine/Source/Programs/SubmitTool/Private/Telemetry/TelemetryService.cpp b/Engine/Source/Programs/SubmitTool/Private/Telemetry/TelemetryService.cpp
--- a/Engine/Source/Programs/SubmitTool/Private/Telemetry/TelemetryService.cpp
+++ b/Engine/Source/Programs/SubmitTool/Private/Telemetry/TelemetryService.cpp
@@ -16,8 +16,20 @@ const TSharedPtr<ITelemetry>& FTelemetryService::Get()
 	return TelemetryInstance;
 }
 
-void FTelemetryService::Init(const FTelemetryParameters& InTelemetryParameters)
+void FTelemetryService::Init(const FTelemetryParameters& InTelemetryParameters, const FString& InStreamName)
 {
+	if (!InTelemetryParameters.RegexStreamAllowlist.IsEmpty())
+	{
+		FRegexPattern Pattern = FRegexPattern(InTelemetryParameters.RegexStreamAllowlist, ERegexPatternFlags::CaseInsensitive);
+		FRegexMatcher regex = FRegexMatcher(Pattern, InStreamName);
+
+		if (!regex.FindNext())
+		{
+			Set(MakeShared<FNullTelemetry>());
+			return;
+		}
+	}
+
 	if (InTelemetryParameters.Instance.Equals(TEXT("StandAlone"), ESearchCase::IgnoreCase))
 	{
 		if (InTelemetryParameters.Url.IsEmpty())
diff --git a/Engine/Source/Programs/SubmitTool/Private/Telemetry/TelemetryService.h b/Engine/Source/Programs/SubmitTool/Private/Telemetry/TelemetryService.h
--- a/Engine/Source/Programs/SubmitTool/Private/Telemetry/TelemetryService.h
+++ b/Engine/Source/Programs/SubmitTool/Private/Telemetry/TelemetryService.h
@@ -13,7 +13,7 @@ class FTelemetryService
 {
 public:
 	static const TSharedPtr<ITelemetry> & Get();
-	static void Init(const FTelemetryParameters& InTelemetryParameters);
+	static void Init(const FTelemetryParameters& InTelemetryParameters, const FString& InStreamName);
 	static void Shutdown();
 	static void BlockFlush(float InTimeout);
 private:
