Adding telemetry attributes:
- Better name for Status -> Success
- Int64 changelist property
- Included validator states in the submit successful event

#Synced-CL 51501979

diff --git a/Engine/Source/Programs/SubmitTool/Private/Logic/Validators/ValidatorBase.cpp b/Engine/Source/Programs/SubmitTool/Private/Logic/Validators/ValidatorBase.cpp
--- a/Engine/Source/Programs/SubmitTool/Private/Logic/Validators/ValidatorBase.cpp
+++ b/Engine/Source/Programs/SubmitTool/Private/Logic/Validators/ValidatorBase.cpp
@@ -390,16 +390,26 @@ const TArray<FAnalyticsEventAttribute> FValidatorBase::GetTelemetryAttributes()
 		StealthStateStr = StaticEnum<EValidationStates>()->GetNameStringByValue(static_cast<int64>(StealthState)).Replace(TEXT("_"), TEXT(" "));
 	}
 
+	FString CLString = ChangelistService->GetCLID();
+	int64 CLNumber;
+	if (!LexTryParseString(CLNumber, *CLString))
+	{
+		CLNumber = 0;
+		CLString = TEXT("0");
+	}
+
 	return MakeAnalyticsEventAttributeArray(
 		TEXT("ValidatorID"), *GetValidatorNameId().ToString(),
 		TEXT("ValidatorName"), *GetValidatorName(),
-		TEXT("Status"), Definition->bIsStealthMode ? bStealthFailure : GetHasPassed(),
+		TEXT("Status"), Definition->bIsStealthMode ? !bStealthFailure : GetHasPassed(),
+		TEXT("Success"), Definition->bIsStealthMode ? !bStealthFailure : GetHasPassed(),
 		TEXT("ValidatorState"), StaticEnum<EValidationStates>()->GetNameStringByValue(static_cast<int64>(State)).Replace(TEXT("_"), TEXT(" ")),
 		TEXT("IsStealthValidator"), Definition->bIsStealthMode,
 		TEXT("StealthActualState"), MoveTemp(StealthStateStr),
 		TEXT("Runtime"), RunTime,
 		TEXT("Stream"), ServiceProvider.Pin()->GetService<ISTSourceControlService>()->GetCurrentStreamName(),
-		TEXT("PendingChangelist"), ChangelistService->GetCLID(),
+		TEXT("PendingChangelist"), MoveTemp(CLString),
+		TEXT("PendingChangeNumber"), CLNumber,
 		TEXT("FailureReason"), StaticEnum<EFailureReason>()->GetNameStringByValue(static_cast<int64>(FailureReason)),
 		TEXT("FromEditor"), FCmdLineParameters::Get().Contains(FSubmitToolCmdLine::EditorFlag),
 		TEXT("Errors"), ProcessMessagesForTelemetry(ErrorListCache),
diff --git a/Engine/Source/Programs/SubmitTool/Private/Logic/Validators/ValidatorBase.h b/Engine/Source/Programs/SubmitTool/Private/Logic/Validators/ValidatorBase.h
--- a/Engine/Source/Programs/SubmitTool/Private/Logic/Validators/ValidatorBase.h
+++ b/Engine/Source/Programs/SubmitTool/Private/Logic/Validators/ValidatorBase.h
@@ -119,6 +119,11 @@ public:
 		return State == EValidationStates::Valid || State == EValidationStates::Skipped || State == EValidationStates::Not_Applicable || State == EValidationStates::Disabled;
 	}
 
+	EValidationStates GetStealthState() const
+	{
+		return StealthState;
+	}
+
 	bool GetIsRunningOrQueued() const
 	{
 		return GetIsRunning() || GetIsQueued();
@@ -248,7 +253,7 @@ protected:
 	bool bForceRun = false;
 	const FSubmitToolParameters& SubmitToolParameters;
 	bool bStealthFailure = false;
-	EValidationStates StealthState;
+	EValidationStates StealthState = EValidationStates::Not_Run;
 	FDateTime Start;
 	float RunTime = 0;
 
diff --git a/Engine/Source/Programs/SubmitTool/Private/Models/ModelInterface.cpp b/Engine/Source/Programs/SubmitTool/Private/Models/ModelInterface.cpp
--- a/Engine/Source/Programs/SubmitTool/Private/Models/ModelInterface.cpp
+++ b/Engine/Source/Programs/SubmitTool/Private/Models/ModelInterface.cpp
@@ -1158,13 +1158,35 @@ void FModelInterface::OnSubmitOperationComplete(const FSCCommand& InCmd)
 					)
 				);
 			}
+		}	
+
+		FString SubmittedCL;
+		int64 CLNumber;
+		if (SubmittedCLString == nullptr || !LexTryParseString(CLNumber, **SubmittedCLString))
+		{
+			CLNumber = 0;
+			SubmittedCL = TEXT("0");
 		}
-	
-		FTelemetryService::Get()->SubmitSucceeded(MakeAnalyticsEventAttributeArray(
+
+
+		TArray<FAnalyticsEventAttribute> SubmitSuccessAttribs = MakeAnalyticsEventAttributeArray(
 			TEXT("Stream"), SourceControlService->GetCurrentStreamName(),
-			TEXT("SubmittedCL"), ((SubmittedCLString == nullptr) ? FString() : *SubmittedCLString)
-		));
+			TEXT("SubmittedCL"), MoveTemp(SubmittedCL),
+			TEXT("SubmittedCLNumber"), MoveTemp(CLNumber));
+
+		for (const TWeakPtr<const FValidatorBase>& Validator : ValidationService->GetTasks())
+		{
+			const TSharedPtr<const FValidatorBase> Ptr = Validator.Pin();
+			if (Ptr != nullptr)
+			{
+				EValidationStates State = Ptr->Definition->bIsStealthMode ? Ptr->GetStealthState() : Ptr->GetState();
+				AppendAnalyticsEventAttributeArray(SubmitSuccessAttribs,
+					FString::Printf(TEXT("V_%s"), *Ptr->GetValidatorNameId().ToString()), StaticEnum<EValidationStates>()->GetNameStringByValue(static_cast<int64>(State)).Replace(TEXT("_"), TEXT(" ")));
+			}
+		}
 
+		FTelemetryService::Get()->SubmitSucceeded(MoveTemp(SubmitSuccessAttribs));
+		
 		// We've submitted, or tried to submit and failed so we only let the user close the app
 		ChangeState(ESubmitToolAppState::Finished);
 		if(FSubmitToolUserPrefs::Get()->bOpenJiraOnSubmit)
diff --git a/Engine/Source/Programs/SubmitTool/Private/Version/AppVersionDefines.h b/Engine/Source/Programs/SubmitTool/Private/Version/AppVersionDefines.h
--- a/Engine/Source/Programs/SubmitTool/Private/Version/AppVersionDefines.h
+++ b/Engine/Source/Programs/SubmitTool/Private/Version/AppVersionDefines.h
@@ -21,7 +21,7 @@
 
 #define SUBMIT_TOOL_VERSION_MAJOR 6
 #define SUBMIT_TOOL_VERSION_MINOR 7
-#define SUBMIT_TOOL_VERSION_PATCH 0
+#define SUBMIT_TOOL_VERSION_PATCH 1
 
 #ifndef SUBMIT_TOOL_CHANGELIST
 #define SUBMIT_TOOL_CHANGELIST_STRING TEXT("NoCL")
