Extend validator launch context filtering to also skip validators that require standalone launch mode when the editor flag is present
Add clarifying comments to the launch context check block

#Synced-CL 51501979

diff --git a/Engine/Source/Programs/SubmitTool/Private/Logic/Validators/ValidatorBase.cpp b/Engine/Source/Programs/SubmitTool/Private/Logic/Validators/ValidatorBase.cpp
--- a/Engine/Source/Programs/SubmitTool/Private/Logic/Validators/ValidatorBase.cpp
+++ b/Engine/Source/Programs/SubmitTool/Private/Logic/Validators/ValidatorBase.cpp
@@ -718,7 +718,12 @@ bool FValidatorBase::AppliesToFile(const FSCFileRef InFile, bool InbAllowIncreme
 		return true;
 	}
 
-	if (Definition->bRequiresSTLaunchedFromEditor && !FCmdLineParameters::Get().Contains(FSubmitToolCmdLine::EditorFlag))
+
+	// Check "-from-editor" cmdline arg
+	// Don't apply if the validator requires being launched from editor and the arg is not present
+	if ((Definition->bRequiresSTLaunchedFromEditor && !FCmdLineParameters::Get().Contains(FSubmitToolCmdLine::EditorFlag))
+		// Don't apply if the validator requires being launched from standalone and the arg is present
+		|| (Definition->bRequiresSTLaunchedStandalone && FCmdLineParameters::Get().Contains(FSubmitToolCmdLine::EditorFlag)))
 	{
 		return false;
 	}
