[SubmitTool] Add NumWarnings and NumErrors attributes to Validator.Finished telemetry event

#rnx
#Synced-CL 51401191

diff --git a/Engine/Source/Programs/SubmitTool/Private/Logic/Validators/ValidatorBase.cpp b/Engine/Source/Programs/SubmitTool/Private/Logic/Validators/ValidatorBase.cpp
--- a/Engine/Source/Programs/SubmitTool/Private/Logic/Validators/ValidatorBase.cpp
+++ b/Engine/Source/Programs/SubmitTool/Private/Logic/Validators/ValidatorBase.cpp
@@ -413,7 +413,9 @@ const TArray<FAnalyticsEventAttribute> FValidatorBase::GetTelemetryAttributes()
 		TEXT("FailureReason"), StaticEnum<EFailureReason>()->GetNameStringByValue(static_cast<int64>(FailureReason)),
 		TEXT("FromEditor"), FCmdLineParameters::Get().Contains(FSubmitToolCmdLine::EditorFlag),
 		TEXT("Errors"), ProcessMessagesForTelemetry(ErrorListCache),
-		TEXT("Warnings"), ProcessMessagesForTelemetry(WarningListCache)
+		TEXT("Warnings"), ProcessMessagesForTelemetry(WarningListCache),
+		TEXT("NumErrors"), ErrorListCache.Num(),
+		TEXT("NumWarnings"), WarningListCache.Num()
 	);
 }
 
