Fix keyboard focus being constantly reset

#Synced-CL 51501979

diff --git a/Engine/Source/Programs/SubmitTool/Private/Models/ModelInterface.cpp b/Engine/Source/Programs/SubmitTool/Private/Models/ModelInterface.cpp
--- a/Engine/Source/Programs/SubmitTool/Private/Models/ModelInterface.cpp
+++ b/Engine/Source/Programs/SubmitTool/Private/Models/ModelInterface.cpp
@@ -501,7 +501,12 @@ void FModelInterface::EvaluateDisabledValidatorsTag()
 
 	if (bUpdateDescription)
 	{
-		FSlateApplication::Get().ClearKeyboardFocus();
+		TSharedPtr<SWidget> PinnedDescriptionBox = DescriptionBox.Pin();
+		if (PinnedDescriptionBox != nullptr && PinnedDescriptionBox->HasKeyboardFocus())
+		{
+			FSlateApplication::Get().ClearKeyboardFocus();
+		}
+
 		ChangelistService->SetCLDescription(DescriptionCopy, true);
 		TagService->ParseCLDescription();
 	}
@@ -548,7 +553,11 @@ void FModelInterface::EvaluateFailedValidatorsTag()
 
 	if (bUpdateDescription)
 	{
-		FSlateApplication::Get().ClearKeyboardFocus();
+		TSharedPtr<SWidget> PinnedDescriptionBox = DescriptionBox.Pin();
+		if (PinnedDescriptionBox != nullptr && PinnedDescriptionBox->HasKeyboardFocus())
+		{
+			FSlateApplication::Get().ClearKeyboardFocus();
+		}
 		ChangelistService->SetCLDescription(DescriptionCopy, true);
 		TagService->ParseCLDescription();
 	}
@@ -569,7 +578,11 @@ void FModelInterface::UpdateSubmitToolTag(bool InbAdd)
 	{
 		if (!HasSubmitToolTag())
 		{
-			FSlateApplication::Get().ClearKeyboardFocus();
+			TSharedPtr<SWidget> PinnedDescriptionBox = DescriptionBox.Pin();
+			if (PinnedDescriptionBox != nullptr && PinnedDescriptionBox->HasKeyboardFocus())
+			{
+				FSlateApplication::Get().ClearKeyboardFocus();
+			}
 			UE_LOG(LogSubmitToolDebug, Log, TEXT("Added Submit Tool tag"));
 			DescriptionCopy.Append(TEXT("\n") + SubmitToolTag);
 			ChangelistService->SetCLDescription(DescriptionCopy, true);
@@ -578,6 +591,12 @@ void FModelInterface::UpdateSubmitToolTag(bool InbAdd)
 	}
 	else if (HasSubmitToolTag())
 	{
+		TSharedPtr<SWidget> PinnedDescriptionBox = DescriptionBox.Pin();
+		if (PinnedDescriptionBox != nullptr && PinnedDescriptionBox->HasKeyboardFocus())
+		{
+			FSlateApplication::Get().ClearKeyboardFocus();
+		}
+
 		const FString VersionlessTag = TEXT("#submittool ");
 
 		if (RemoveTagFromDescription(VersionlessTag, DescriptionCopy))
diff --git a/Engine/Source/Programs/SubmitTool/Private/Models/ModelInterface.h b/Engine/Source/Programs/SubmitTool/Private/Models/ModelInterface.h
--- a/Engine/Source/Programs/SubmitTool/Private/Models/ModelInterface.h
+++ b/Engine/Source/Programs/SubmitTool/Private/Models/ModelInterface.h
@@ -85,6 +85,7 @@ public:
 
 	void SetMainTab(TSharedPtr<SDockTab> InMainTab)									{ MainTab = InMainTab; }
 	TWeakPtr<SDockTab> GetMainTab() const											{ return MainTab; }
+	void SetDescriptionBox(const TWeakPtr<SWidget> InDescriptionBox)				{ DescriptionBox = InDescriptionBox; }
 
 	void ApplyTag(const FString& tagID) const										{ TagService->ApplyTag(tagID);}
 	void ApplyTag(const FTag& tag) const											{ TagService->ApplyTag(const_cast<FTag&>(tag)); ValidationService->InvalidateForChanges(ETaskArea::Changelist); }
@@ -284,6 +285,7 @@ private:
 	TSharedPtr<FJiraService> JiraService;
 	TSharedPtr<FPreflightService> PreflightService;
 	TWeakPtr<SDockTab> MainTab;
+	TWeakPtr<SWidget> DescriptionBox;
 	TSharedPtr<FTasksService> PresubmitOperationsService;
 	TSharedPtr<FIntegrationService> IntegrationService;
 	TSharedPtr<FSwarmService> SwarmService;
diff --git a/Engine/Source/Programs/SubmitTool/Private/View/Widgets/SubmitToolWidget.cpp b/Engine/Source/Programs/SubmitTool/Private/View/Widgets/SubmitToolWidget.cpp
--- a/Engine/Source/Programs/SubmitTool/Private/View/Widgets/SubmitToolWidget.cpp
+++ b/Engine/Source/Programs/SubmitTool/Private/View/Widgets/SubmitToolWidget.cpp
@@ -196,7 +196,9 @@ void SubmitToolWidget::Construct(const FArguments& InArgs)
 		.AutoWrapText(true)
 		.OnIsTypedCharValid(FOnIsTypedCharValid::CreateLambda([](const TCHAR) { return true; }))
 		.IsReadOnly_Lambda([this]{ return !FModelInterface::GetInputEnabled() || ModelInterface->IsDescriptionGenerationRunning();});
-
+	
+	ModelInterface->SetDescriptionBox(DescriptionBox);
+	
 	FString PerforceClientName;
 	FCmdLineParameters::Get().GetValue(FSubmitToolCmdLine::P4Client, PerforceClientName);
 
