Move bIncludeWarningsInErrorSummary to the base validator definition and default it to true
Refactor error and warning telemetry dispatching to always execute on the game thread
Add warning telemetry support to the executable validator and separate error/warning handling paths

#Synced-CL 51501979

diff --git a/Engine/Source/Developer/SubmitToolCore/Public/Logic/Validators/ValidatorDefinition.h b/Engine/Source/Developer/SubmitToolCore/Public/Logic/Validators/ValidatorDefinition.h
--- a/Engine/Source/Developer/SubmitToolCore/Public/Logic/Validators/ValidatorDefinition.h
+++ b/Engine/Source/Developer/SubmitToolCore/Public/Logic/Validators/ValidatorDefinition.h
@@ -291,6 +291,12 @@ struct FValidatorDefinition
 
 	UPROPERTY()
 	FString PreflightTemplate;
+
+	///
+	/// Whether to include or not the output warnings in the error summary window
+	/// 
+	UPROPERTY()
+	bool bIncludeWarningsInErrorSummary = true;
 };
 
 USTRUCT()
@@ -413,12 +419,6 @@ struct FValidatorRunExecutableDefinition : public FValidatorDefinition
 	UPROPERTY()
 	FString EnableWarningsAsErrorsAnchor;
 
-	///
-	/// Whether to include or not the output warnings in the error summary window
-	/// 
-	UPROPERTY()
-	bool bIncludeWarningsInErrorSummary = false;
-
 	///
 	///	Regex for identifying errors from the output of a process
 	/// 
diff --git a/Engine/Source/Programs/SubmitTool/Private/Logic/Validators/ValidatorBase.cpp b/Engine/Source/Programs/SubmitTool/Private/Logic/Validators/ValidatorBase.cpp
--- a/Engine/Source/Programs/SubmitTool/Private/Logic/Validators/ValidatorBase.cpp
+++ b/Engine/Source/Programs/SubmitTool/Private/Logic/Validators/ValidatorBase.cpp
@@ -501,16 +501,8 @@ void FValidatorBase::LogFailure(const FString& FormattedMsg) const
 			UE_LOG(LogValidatorsResult, Error, TEXT("%s"), *FormattedMsg);
 		}
 
-	
-
 		FScopeLock Lock(&ErrorLogsMutex);
 		ErrorListCache.Add(FormattedMsg);
-
-		FString TelemetryMsg = FormattedMsg;
-		ProcessSingleMessageForTelemetry(TelemetryMsg);
-
-		FTelemetryService::Get()->CustomEvent(TEXT("SubmitTool.Task.Error"), GetSingleMessageTelemetryAttributes(TelemetryMsg));
-
 	}
 	else
 	{
@@ -526,13 +518,17 @@ void FValidatorBase::LogFailure(const FString& FormattedMsg) const
 
 		FScopeLock Lock(&ErrorLogsMutex);
 		WarningListCache.Add(FormattedMsg);
+	}
 
+	// Launch the telemetry on the game thread
+	UE::Tasks::Launch(UE_SOURCE_LOCATION, [this, FormattedMsg]() mutable
+	{
 		FString TelemetryMsg = FormattedMsg;
 		ProcessSingleMessageForTelemetry(TelemetryMsg);
+		FString EventName = Definition->IsRequired ? TEXT("SubmitTool.Task.Error") : TEXT("SubmitTool.Task.Warning");
 
-		FTelemetryService::Get()->CustomEvent(TEXT("SubmitTool.Task.Warning"), GetSingleMessageTelemetryAttributes(TelemetryMsg));
-	}
-
+		FTelemetryService::Get()->CustomEvent(MoveTemp(EventName), GetSingleMessageTelemetryAttributes(TelemetryMsg));
+	}, LowLevelTasks::ETaskPriority::Normal, UE::Tasks::EExtendedTaskPriority::GameThreadNormalPri);
 }
 
 void FValidatorBase::SetValidationFinished(EValidationStates NewState)
@@ -687,10 +683,13 @@ void FValidatorBase::PrintErrorSummary() const
 			}
 		}
 
-		for (const FString& WarningStr : WarningListCache)
+		if (Definition->bIncludeWarningsInErrorSummary)
 		{
-			UE_LOG(LogValidators, Warning, TEXT("%s"), *WarningStr);
-			UE_LOG(LogValidatorsResult, Warning, TEXT("%s"), *WarningStr);
+			for (const FString& WarningStr : WarningListCache)
+			{
+				UE_LOG(LogValidators, Warning, TEXT("%s"), *WarningStr);
+				UE_LOG(LogValidatorsResult, Warning, TEXT("%s"), *WarningStr);
+			}
 		}
 	}
 }
diff --git a/Engine/Source/Programs/SubmitTool/Private/Logic/Validators/ValidatorRunExecutable.cpp b/Engine/Source/Programs/SubmitTool/Private/Logic/Validators/ValidatorRunExecutable.cpp
--- a/Engine/Source/Programs/SubmitTool/Private/Logic/Validators/ValidatorRunExecutable.cpp
+++ b/Engine/Source/Programs/SubmitTool/Private/Logic/Validators/ValidatorRunExecutable.cpp
@@ -19,6 +19,7 @@
 #include "Serialization/JsonReader.h"
 #include "Serialization/JsonSerializer.h"
 #include "Serialization/JsonWriter.h"
+#include "Telemetry/TelemetryService.h"
 
 int32 FValidatorRunExecutable::ProcessCounter = 0;
 FString FValidatorRunExecutable::GetProcessName(FStringView ExecutablePath)
@@ -236,7 +237,14 @@ bool FValidatorRunExecutable::QueueProcess(const FString& InId, const FString& L
 
 void FValidatorRunExecutable::OnProcessOutputLine(const FString& ProcessId, const FString& Line, const EProcessOutputType& OutputType)
 {
+
 	const FString& FormattedMessage = FString::Printf(TEXT("[%s] [%s] %s"), *GetValidatorName(), *ProcessId, *Line);
+	if (OutputType == EProcessOutputType::ProcessError)
+	{
+		LogFailure(FormattedMessage);
+		ProcessesData[ProcessId].ErrorList.Add(FormattedMessage);
+		return;
+	}
 
 	for (const FRegexLockdownTrigger& LockdownTrigger : GetTypedDefinition<FValidatorRunExecutableDefinition>()->LogLockdownTriggers)
 	{
@@ -249,26 +257,35 @@ void FValidatorRunExecutable::OnProcessOutputLine(const FString& ProcessId, cons
 	}
 
 	bool bIsWarning = false;
-	if(OutputType == EProcessOutputType::ProcessError || IsLineAnError(Line, bIsWarning, ProcessesData[ProcessId].bIgnoringOutputError, ProcessesData[ProcessId].bIgnoringWarningsAsErrors))
+	bool bIsError = IsLineAnError(Line, bIsWarning, ProcessesData[ProcessId].bIgnoringOutputError, ProcessesData[ProcessId].bIgnoringWarningsAsErrors);
+
+	if(bIsError)
 	{
+		// LogFailure already handles telemetry
 		LogFailure(FormattedMessage);
 		ProcessesData[ProcessId].ErrorList.Add(FormattedMessage);
 	}
-	else
+	else if (bIsWarning)
 	{
-		if (bIsWarning && GetTypedDefinition<FValidatorRunExecutableDefinition>()->bIncludeWarningsInErrorSummary)
-		{
-			FScopeLock Lock(&ErrorLogsMutex);
-			ProcessesData[ProcessId].WarningList.Add(FormattedMessage);
-			WarningListCache.Add(FormattedMessage);
-			UE_LOG(LogValidatorsResult, Warning, TEXT("%s"), *FormattedMessage);
-			UE_LOG(LogValidators, Warning, TEXT("%s"), *FormattedMessage);
+		FScopeLock Lock(&ErrorLogsMutex);
+		ProcessesData[ProcessId].WarningList.Add(FormattedMessage);
+		WarningListCache.Add(FormattedMessage);
 
-		}
-		else
+		UE_LOG(LogValidatorsResult, Warning, TEXT("%s"), *FormattedMessage);
+		UE_LOG(LogValidators, Warning, TEXT("%s"), *FormattedMessage);
+
+		// Handle telemetry message in the game thread for this warning
+		UE::Tasks::Launch(UE_SOURCE_LOCATION, [this, FormattedMessage]() mutable
 		{
-			UE_LOG(LogValidators, Log, TEXT("%s"), *FormattedMessage);
-		}		
+			FString TelemetryMsg = FormattedMessage;
+			ProcessSingleMessageForTelemetry(TelemetryMsg);
+
+			FTelemetryService::Get()->CustomEvent(TEXT("SubmitTool.Task.Warning"), GetSingleMessageTelemetryAttributes(TelemetryMsg));
+		}, LowLevelTasks::ETaskPriority::Normal, UE::Tasks::EExtendedTaskPriority::GameThreadNormalPri);
+	}
+	else
+	{
+		UE_LOG(LogValidators, Log, TEXT("%s"), *FormattedMessage);
 	}
 }
 
