Remove keyboard focus clearing logic from the submit tool tag handling path to prevent losing focus on edit

#Synced-CL 51501979
#Synced-CL 51501979

diff --git a/Engine/Source/Programs/SubmitTool/Private/Models/ModelInterface.cpp b/Engine/Source/Programs/SubmitTool/Private/Models/ModelInterface.cpp
--- a/Engine/Source/Programs/SubmitTool/Private/Models/ModelInterface.cpp
+++ b/Engine/Source/Programs/SubmitTool/Private/Models/ModelInterface.cpp
@@ -591,12 +591,6 @@ void FModelInterface::UpdateSubmitToolTag(bool InbAdd)
 	}
 	else if (HasSubmitToolTag())
 	{
-		TSharedPtr<SWidget> PinnedDescriptionBox = DescriptionBox.Pin();
-		if (PinnedDescriptionBox != nullptr && PinnedDescriptionBox->HasKeyboardFocus())
-		{
-			FSlateApplication::Get().ClearKeyboardFocus();
-		}
-
 		const FString VersionlessTag = TEXT("#submittool ");
 
 		if (RemoveTagFromDescription(VersionlessTag, DescriptionCopy))
