Submit Tool absolute log

#rnx
#Synced-CL 0

diff --git a/Engine/Source/Programs/SubmitTool/Private/SubmitToolApp.cpp b/Engine/Source/Programs/SubmitTool/Private/SubmitToolApp.cpp
--- a/Engine/Source/Programs/SubmitTool/Private/SubmitToolApp.cpp
+++ b/Engine/Source/Programs/SubmitTool/Private/SubmitToolApp.cpp
@@ -180,6 +180,8 @@ int RunSubmitTool(const TCHAR* CommandLine)
 	TRACE_CPUPROFILER_EVENT_SCOPE(RunSubmitTool);
 	FTaskTagScope TaskTagScope(ETaskTag::EGameThread);
 
+	FPlatformOutputDevices::SetupOutputDevices();
+
 	// need to make sure the cwd is correct before doing anything else
 	FPlatformProcess::SetCurrentWorkingDirectoryToBaseDir();
 
