[SubmitTool] Add Correlation GUID to each Validator run to use for telemetry

#rnx

#Synced-CL 51586555

#Synced-CL 51586555

diff --git a/Engine/Source/Programs/SubmitTool/Private/Logic/Validators/EditorCommandletValidator.cpp b/Engine/Source/Programs/SubmitTool/Private/Logic/Validators/EditorCommandletValidator.cpp
--- a/Engine/Source/Programs/SubmitTool/Private/Logic/Validators/EditorCommandletValidator.cpp
+++ b/Engine/Source/Programs/SubmitTool/Private/Logic/Validators/EditorCommandletValidator.cpp
@@ -69,6 +69,11 @@ bool FEditorCommandletValidator::Validate(const FString& InCLDescription, const
 			EditorArguments += FString::Printf(TEXT(" -UserDir=\"%s\""), *TempUserDir);
 		}
 
+		if (CorrelationId.IsValid())
+		{
+			EditorArguments += FString::Printf(TEXT(" -SubmitToolCorrelationId=\"%s\""), *CorrelationId.ToString());
+		}
+
 		bProcessesStarted |= QueueProcess(EditorParametersSet.Key, EditorParametersSet.Value.EditorExePath, EditorArguments);
 	}
 
diff --git a/Engine/Source/Programs/SubmitTool/Private/Logic/Validators/ValidatorBase.cpp b/Engine/Source/Programs/SubmitTool/Private/Logic/Validators/ValidatorBase.cpp
--- a/Engine/Source/Programs/SubmitTool/Private/Logic/Validators/ValidatorBase.cpp
+++ b/Engine/Source/Programs/SubmitTool/Private/Logic/Validators/ValidatorBase.cpp
@@ -26,6 +26,7 @@ FValidatorBase::FValidatorBase(const FName& InNameId, const FSubmitToolParameter
 {
 	ParseDefinition(InDefinition);
 	ValidatorName = Definition->CustomName.IsEmpty() ? ValidatorNameID.ToString() : Definition->CustomName;
+	CorrelationId = FGuid::NewGuid();
 }
 
 void FValidatorBase::ParseDefinition(const FString& InDefinition)
@@ -77,6 +78,7 @@ void FValidatorBase::StartValidation()
 	FailureReason = EFailureReason::None;
 	ErrorListCache.Reset();
 	WarningListCache.Reset();
+	CorrelationId = FGuid::NewGuid();
 
 	const TSharedPtr<FChangelistService>& ChangelistService = ServiceProvider.Pin()->GetService<FChangelistService>();
 	const TSharedPtr<FTagService>& TagService = ServiceProvider.Pin()->GetService<FTagService>();
@@ -185,6 +187,7 @@ void FValidatorBase::StartValidation()
 				UE_LOG(LogValidators, Log, TEXT("[%s] Skipping Files because they were already validated in a previous execution:\n%s"), *GetValidatorName(), *FileList);
 			}
 
+			UE_LOG(LogValidators, Log, TEXT("[%s] CorrelationId: %s"), *GetValidatorName(), *CorrelationId.ToString());
 			if (!Validate(ChangelistService->GetCLDescription(), FilteredFiles, TagService->GetTagsArray()))
 			{
 				if(FailureReason == EFailureReason::None)
@@ -415,7 +418,8 @@ const TArray<FAnalyticsEventAttribute> FValidatorBase::GetTelemetryAttributes()
 		TEXT("Errors"), ProcessMessagesForTelemetry(ErrorListCache),
 		TEXT("Warnings"), ProcessMessagesForTelemetry(WarningListCache),
 		TEXT("NumErrors"), ErrorListCache.Num(),
-		TEXT("NumWarnings"), WarningListCache.Num()
+		TEXT("NumWarnings"), WarningListCache.Num(),
+		TEXT("CorrelationId"), CorrelationId.ToString()
 	);
 }
 
diff --git a/Engine/Source/Programs/SubmitTool/Private/Logic/Validators/ValidatorBase.h b/Engine/Source/Programs/SubmitTool/Private/Logic/Validators/ValidatorBase.h
--- a/Engine/Source/Programs/SubmitTool/Private/Logic/Validators/ValidatorBase.h
+++ b/Engine/Source/Programs/SubmitTool/Private/Logic/Validators/ValidatorBase.h
@@ -238,6 +238,8 @@ public:
 	FOnValidatorFinished OnValidationFinished;
 	FOnLockdownIdsChanged OnLockdownIdsChanged;
 
+	FGuid CorrelationId; // Automatically set to a new GUID each time before Validate is called
+
 protected:
 	FName ValidatorNameID;
 	FString ValidatorName;
