Add a custom output log widget to SubmitTool that persists log history before log widget is created
Replace the generic engine output log with a dedicated submit tool widget subclass, enabling scroll-to-end behavior and proper tab focus on submit and validation events
Broadcast a submit finished delegate with success state to allow the UI to react accordingly

#Synced-CL 51711301

diff --git a/Engine/Source/Programs/SubmitTool/Private/Logging/SubmitToolOutputLog.cpp b/Engine/Source/Programs/SubmitTool/Private/Logging/SubmitToolOutputLog.cpp
new file mode 100644
--- /dev/null
+++ b/Engine/Source/Programs/SubmitTool/Private/Logging/SubmitToolOutputLog.cpp
@@ -0,0 +1,8 @@
+// Copyright Epic Games, Inc. All Rights Reserved.
+
+#include "SubmitToolOutputLog.h"
+
+void SSubmitToolOutputLog::ScrollToEnd()
+{
+	RequestForceScroll(false);
+}
diff --git a/Engine/Source/Programs/SubmitTool/Private/Logging/SubmitToolOutputLog.h b/Engine/Source/Programs/SubmitTool/Private/Logging/SubmitToolOutputLog.h
new file mode 100644
--- /dev/null
+++ b/Engine/Source/Programs/SubmitTool/Private/Logging/SubmitToolOutputLog.h
@@ -0,0 +1,13 @@
+// Copyright Epic Games, Inc. All Rights Reserved.
+
+#pragma once
+
+#include "SOutputLog.h"
+
+class SSubmitToolOutputLog : public SOutputLog
+{
+public:
+	using FArguments = SOutputLog::FArguments;
+
+	void ScrollToEnd();
+};
diff --git a/Engine/Source/Programs/SubmitTool/Private/Logging/SubmitToolOutputLogHistory.h b/Engine/Source/Programs/SubmitTool/Private/Logging/SubmitToolOutputLogHistory.h
new file mode 100644
--- /dev/null
+++ b/Engine/Source/Programs/SubmitTool/Private/Logging/SubmitToolOutputLogHistory.h
@@ -0,0 +1,50 @@
+// Copyright Epic Games, Inc. All Rights Reserved.
+
+#pragma once
+
+#include "HAL/IConsoleManager.h"
+#include "Misc/OutputDevice.h"
+#include "Logging/SubmitToolLog.h"
+#include "SOutputLog.h"
+
+class FSubmitToolOutputLogHistory : public FOutputDevice
+{
+public:
+	FSubmitToolOutputLogHistory()
+	{
+		GLog->AddOutputDevice(this);
+	}
+
+	virtual ~FSubmitToolOutputLogHistory()
+	{
+		if (GLog != nullptr)
+		{
+			GLog->RemoveOutputDevice(this);
+		}
+	}
+
+	FSubmitToolOutputLogHistory(const FSubmitToolOutputLogHistory&) = delete;
+	FSubmitToolOutputLogHistory(FSubmitToolOutputLogHistory&&) = delete;
+
+	const TArray<TSharedPtr<FOutputLogMessage>>& GetMessages() const
+	{
+		return Messages;
+	}
+
+private:
+	virtual void Serialize(const TCHAR* V, ELogVerbosity::Type Verbosity, const FName& Category) override
+	{
+		// this log history is meant to be short lived, its purpose is to store messages before the UI widget is created
+		if (Messages.Num() < 2000)
+		{
+			SOutputLog::CreateLogMessages(V, Verbosity, Category, Messages);
+		}
+		else
+		{
+			GLog->RemoveOutputDevice(this);
+			UE_LOG(LogSubmitToolDebug, Error, TEXT("Output log history went over 2000 messages, unregistering"));
+		}
+	}
+
+	TArray<TSharedPtr<FOutputLogMessage>> Messages;
+};
diff --git a/Engine/Source/Programs/SubmitTool/Private/Models/ModelInterface.cpp b/Engine/Source/Programs/SubmitTool/Private/Models/ModelInterface.cpp
--- a/Engine/Source/Programs/SubmitTool/Private/Models/ModelInterface.cpp
+++ b/Engine/Source/Programs/SubmitTool/Private/Models/ModelInterface.cpp
@@ -179,6 +179,7 @@ FModelInterface::~FModelInterface()
 {
 	PrepareSubmitCallBack.Clear();
 	FileRefreshedCallback.Clear();
+	SubmitFinishedCallback.Clear();
 	ValidationService->OnTasksRunResultUpdated.Remove(OnValidationStateUpdatedHandle);
 	ValidationService->OnTasksQueueFinished.Remove(OnValidationFinishedHandle);
 	ValidationService->OnSingleTaskFinished.Remove(OnSingleValidationFinishedHandle);
@@ -1225,6 +1226,8 @@ void FModelInterface::OnSubmitOperationComplete(const FSCCommand& InCmd)
 			}
 		}
 
+		SubmitFinishedCallback.Broadcast(true);
+
 		if (FSubmitToolUserPrefs::Get()->bCloseOnSubmit)
 		{
 			MainTab.Pin()->RequestCloseTab();
@@ -1233,6 +1236,7 @@ void FModelInterface::OnSubmitOperationComplete(const FSCCommand& InCmd)
 	else
 	{
 		ChangeState(ESubmitToolAppState::WaitingUserInput);
+		SubmitFinishedCallback.Broadcast(false);
 	}
 }
 
diff --git a/Engine/Source/Programs/SubmitTool/Private/Models/ModelInterface.h b/Engine/Source/Programs/SubmitTool/Private/Models/ModelInterface.h
--- a/Engine/Source/Programs/SubmitTool/Private/Models/ModelInterface.h
+++ b/Engine/Source/Programs/SubmitTool/Private/Models/ModelInterface.h
@@ -72,6 +72,7 @@ namespace SubmitToolAppState {
 
 DECLARE_MULTICAST_DELEGATE(FPreSubmitCallBack)
 DECLARE_MULTICAST_DELEGATE(FFilesRefresh)
+DECLARE_MULTICAST_DELEGATE_OneParam(FOnSubmit, bool /*bSuccess*/)
 DECLARE_MULTICAST_DELEGATE_TwoParams(FOnStateChanged, const ESubmitToolAppState /*InFrom*/, const ESubmitToolAppState /*InTo*/)
 
 class FModelInterface
@@ -222,6 +223,7 @@ public:
 
 
 	static FOnStateChanged OnStateChanged;
+	FOnSubmit SubmitFinishedCallback;
 	FPreSubmitCallBack PrepareSubmitCallBack;
 	FFilesRefresh FileRefreshedCallback;
 	bool bSubmitOnSuccessfulValidation = false;
diff --git a/Engine/Source/Programs/SubmitTool/Private/SubmitToolApp.cpp b/Engine/Source/Programs/SubmitTool/Private/SubmitToolApp.cpp
--- a/Engine/Source/Programs/SubmitTool/Private/SubmitToolApp.cpp
+++ b/Engine/Source/Programs/SubmitTool/Private/SubmitToolApp.cpp
@@ -5,6 +5,7 @@
 #include "SubmitToolCoreUtils.h"
 #include "Models/SubmitToolUserPrefs.h"
 #include "Logic/Services/SourceControl/SubmitToolPerforce.h"
+#include "Logging/SubmitToolOutputLogHistory.h"
 
 #include "GenericPlatform/GenericPlatformMisc.h"
 #include "RequiredProgramMainCPPInclude.h"
@@ -50,7 +51,7 @@ public:
 		, MainDockTab(SNew(SDockTab))
 		, CommandList(MakeShared<FUICommandList>())
 	{ }
-	void Build(FSubmitToolUserPrefs* UserPrefs);
+	void Build(FSubmitToolUserPrefs* UserPrefs, TSharedPtr<FSubmitToolOutputLogHistory> InLogHistory);
 
 private:
 
@@ -70,7 +71,7 @@ TUniquePtr<FSlateUI> MakeSlateUI(FModelInterface* ModelInterface)
 	return MakeUnique<FSlateUI>(ModelInterface);
 }
 
-void FSlateUI::Build(FSubmitToolUserPrefs* UserPrefs)
+void FSlateUI::Build(FSubmitToolUserPrefs* UserPrefs, TSharedPtr<FSubmitToolOutputLogHistory> InLogHistory)
 {
 	TRACE_CPUPROFILER_EVENT_SCOPE(BuildSlateUI);
 
@@ -95,7 +96,7 @@ void FSlateUI::Build(FSubmitToolUserPrefs* UserPrefs)
 
 	FName TabName = FName("Submit Tool");
 
-	TabManager->RegisterTabSpawner(TabName, FOnSpawnTab::CreateLambda([&Window = Window](const FSpawnTabArgs& SpawnArgs) { return Window.BuildMainTab(SpawnArgs.GetOwnerWindow()); }));
+	TabManager->RegisterTabSpawner(TabName, FOnSpawnTab::CreateLambda([&Window = Window, InLogHistory](const FSpawnTabArgs& SpawnArgs) mutable { return Window.BuildMainTab(SpawnArgs.GetOwnerWindow(), InLogHistory); }));
 	TabManager->RegisterDefaultTabWindowSize(TabName, FVector2D(1024, 600));
 
 	Tab = TabManager->TryInvokeTab(TabName);
@@ -181,6 +182,7 @@ int RunSubmitTool(const TCHAR* CommandLine)
 	FTaskTagScope TaskTagScope(ETaskTag::EGameThread);
 
 	FPlatformOutputDevices::SetupOutputDevices();
+	TSharedPtr<FSubmitToolOutputLogHistory> LogHistory = MakeShared<FSubmitToolOutputLogHistory>();
 
 	// need to make sure the cwd is correct before doing anything else
 	FPlatformProcess::SetCurrentWorkingDirectoryToBaseDir();
@@ -251,7 +253,7 @@ int RunSubmitTool(const TCHAR* CommandLine)
 
 		// Build the slate ui
 		TUniquePtr<FSlateUI> SlateUI = MakeSlateUI(ModelInterface.Get());
-		SlateUI->Build(UserPrefs.Get());
+		SlateUI->Build(UserPrefs.Get(), MoveTemp(LogHistory));
 
 		// Test mode, attach to debugger, only non-shipping
 #if !UE_BUILD_SHIPPING
diff --git a/Engine/Source/Programs/SubmitTool/Private/View/SubmitToolWindow.cpp b/Engine/Source/Programs/SubmitTool/Private/View/SubmitToolWindow.cpp
--- a/Engine/Source/Programs/SubmitTool/Private/View/SubmitToolWindow.cpp
+++ b/Engine/Source/Programs/SubmitTool/Private/View/SubmitToolWindow.cpp
@@ -14,7 +14,7 @@
 SubmitToolWindow::SubmitToolWindow(FModelInterface* modelInterface) : ModelInterface(modelInterface)
 {}
 
-TSharedRef<SDockTab> SubmitToolWindow::BuildMainTab(TSharedPtr<SWindow> InParentWindow)
+TSharedRef<SDockTab> SubmitToolWindow::BuildMainTab(TSharedPtr<SWindow> InParentWindow, TSharedPtr<FSubmitToolOutputLogHistory> InLogHistory)
 {
 	//~~~~~~~~~~~~~~~~~~~~~~~~~~~~~~~~~~~~~~~~~~
 	// Create a submit files window
@@ -33,7 +33,7 @@ TSharedRef<SDockTab> SubmitToolWindow::BuildMainTab(TSharedPtr<SWindow> InParent
 
 	if (ModelInterface->CheckForNewVersion())
 	{
-		CreateAutoUpdateSubmitToolContent(InParentWindow);
+		CreateAutoUpdateSubmitToolContent(InParentWindow, InLogHistory);
 		if (FSubmitToolUserPrefs::Get()->bAutoUpdate)
 		{
 			ModelInterface->InstallLatestVersion();
@@ -41,27 +41,27 @@ TSharedRef<SDockTab> SubmitToolWindow::BuildMainTab(TSharedPtr<SWindow> InParent
 	}
 	else
 	{
-		CreateMainSubmitToolContent(InParentWindow);
+		CreateMainSubmitToolContent(InParentWindow, InLogHistory);
 	}
 
 	return MainTab.ToSharedRef();
 }
 
-void SubmitToolWindow::CreateAutoUpdateSubmitToolContent(TSharedPtr<SWindow> InParentWindow)
+void SubmitToolWindow::CreateAutoUpdateSubmitToolContent(TSharedPtr<SWindow> InParentWindow, TSharedPtr<FSubmitToolOutputLogHistory> InLogHistory)
 {
 	TSharedRef<SAutoUpdateWidget> AutoUpdateWidget =
 		SNew(SAutoUpdateWidget)
 		.ModelInterface(ModelInterface)
-		.OnAutoUpdateCancelled_Lambda([InParentWindow, this]() { CreateMainSubmitToolContent(InParentWindow); });
+		.OnAutoUpdateCancelled_Lambda([InParentWindow, this, InLogHistory]() mutable { CreateMainSubmitToolContent(InParentWindow, InLogHistory); });
 
 	ModelInterface->SetMainTab(MainTab);
 	MainTab->SetContent(AutoUpdateWidget);
 }
 
-void SubmitToolWindow::CreateMainSubmitToolContent(TSharedPtr<SWindow> InParentWindow)
+void SubmitToolWindow::CreateMainSubmitToolContent(TSharedPtr<SWindow> InParentWindow, TSharedPtr<FSubmitToolOutputLogHistory> InLogHistory)
 {
 	TSharedRef<SubmitToolWidget> SourceControlWidget =
-		SNew(SubmitToolWidget)
+		SNew(SubmitToolWidget, InLogHistory)
 		.ParentWindow(InParentWindow)
 		.ParentTab(MainTab)
 		.ModelInterface(ModelInterface);
diff --git a/Engine/Source/Programs/SubmitTool/Private/View/SubmitToolWindow.h b/Engine/Source/Programs/SubmitTool/Private/View/SubmitToolWindow.h
--- a/Engine/Source/Programs/SubmitTool/Private/View/SubmitToolWindow.h
+++ b/Engine/Source/Programs/SubmitTool/Private/View/SubmitToolWindow.h
@@ -6,21 +6,22 @@
 #include "Widgets/Docking/SDockTab.h"
 
 class FModelInterface;
+class FSubmitToolOutputLogHistory;
 
 class SubmitToolWindow
 {
 public:
 	SubmitToolWindow(FModelInterface* modelInterface);
 
-	TSharedRef<SDockTab> BuildMainTab(TSharedPtr<SWindow> InParentWindow);
+	TSharedRef<SDockTab> BuildMainTab(TSharedPtr<SWindow> InParentWindow, TSharedPtr<FSubmitToolOutputLogHistory> InLogHistory);
 
 private:
 	TSharedPtr<SDockTab> MainTab;
 	
 	bool OnCanCloseTab();
 	
-	void CreateAutoUpdateSubmitToolContent(TSharedPtr<SWindow> InParentWindow);
-	void CreateMainSubmitToolContent(TSharedPtr<SWindow> InParentWindow);
+	void CreateAutoUpdateSubmitToolContent(TSharedPtr<SWindow> InParentWindow, TSharedPtr<FSubmitToolOutputLogHistory> InLogHistory);
+	void CreateMainSubmitToolContent(TSharedPtr<SWindow> InParentWindow, TSharedPtr<FSubmitToolOutputLogHistory> InLogHistory);
 
 	FModelInterface* ModelInterface;
 };
diff --git a/Engine/Source/Programs/SubmitTool/Private/View/Widgets/SubmitToolWidget.cpp b/Engine/Source/Programs/SubmitTool/Private/View/Widgets/SubmitToolWidget.cpp
--- a/Engine/Source/Programs/SubmitTool/Private/View/Widgets/SubmitToolWidget.cpp
+++ b/Engine/Source/Programs/SubmitTool/Private/View/Widgets/SubmitToolWidget.cpp
@@ -26,6 +26,8 @@
 #include "Widgets/Images/SThrobber.h"
 #include "Widgets/Docking/SDockTab.h"
 
+#include "Logging/SubmitToolOutputLogHistory.h"
+#include "Logging/SubmitToolOutputLog.h"
 #include "Models/ModelInterface.h"
 #include "Models/SubmitToolUserPrefs.h"
 #include "TagSectionWidget.h"
@@ -73,7 +75,7 @@ void SubmitToolWidget::HandleApplicationActivationStateChanged(bool bActive)
 	}
 }
 
-void SubmitToolWidget::Construct(const FArguments& InArgs)
+void SubmitToolWidget::Construct(const FArguments& InArgs, TSharedPtr<FSubmitToolOutputLogHistory> InLogHistory)
 {
 	ParentTab = InArgs._ParentTab.Get();
 	ModelInterface = InArgs._ModelInterface;
@@ -84,7 +86,7 @@ void SubmitToolWidget::Construct(const FArguments& InArgs)
 
 	FSlateApplication::Get().OnApplicationActivationStateChanged().AddRaw(this, &SubmitToolWidget::HandleApplicationActivationStateChanged);
 	OnValidatorFinishedHandle = ModelInterface->AddSingleValidatorFinishedCallback(FOnSingleTaskFinished::FDelegate::CreateRaw(this, &SubmitToolWidget::OnSingleValidatorFinished));
-	OnValidationUpdateHandle = ModelInterface->AddValidationUpdatedCallback(FOnTaskFinished::FDelegate::CreateRaw(this, &SubmitToolWidget::OnValidationUpdated));
+	OnValidationUpdateHandle = ModelInterface->AddValidationUpdatedCallback(FOnTaskRunStateChanged::FDelegate::CreateRaw(this, &SubmitToolWidget::OnValidationUpdated));
 	OnCLDescriptionUpdatedHandle = ModelInterface->GetCLDescriptionUpdatedDelegate().Add(FOnCLDescriptionUpdated::FDelegate::CreateRaw(this, &SubmitToolWidget::OnCLDescriptionUpdated));
 
 	TSharedPtr<SVerticalBox> Contents;
@@ -379,8 +381,9 @@ void SubmitToolWidget::Construct(const FArguments& InArgs)
 	
 	P4SectionSlot->SetSizeValue(FSubmitToolUserPrefs::Get()->P4SectionSize);
 
-	TSharedPtr<SWidget> LogSection = BuildOutputLogWidget();
+	TSharedPtr<SWidget> LogSection = BuildOutputLogWidget(InLogHistory);
 	TSharedPtr<SWidget> FilesInCLSection = BuildFilesInCLWidget();
+	InLogHistory.Reset();
 
 	// Files & Validators
 	ValidatorSectionSlot = Splitter->AddSlot()		
@@ -672,7 +675,7 @@ TSharedRef<SHorizontalBox> SubmitToolWidget::BuildButtonRow()
 		];
 }
 
-TSharedRef<SWidget> SubmitToolWidget::BuildOutputLogWidget()
+TSharedRef<SWidget> SubmitToolWidget::BuildOutputLogWidget(const TSharedPtr<FSubmitToolOutputLogHistory>& InLogHistory)
 {
 	/*** Output Log Widget ***/
 	FOutputLogModule& OutputLogModule = FModuleManager::Get().LoadModuleChecked<FOutputLogModule>("OutputLog");
@@ -723,26 +726,17 @@ TSharedRef<SWidget> SubmitToolWidget::BuildOutputLogWidget()
 	LogTabManager->SetAllowWindowMenuBar(true);
 	LogTabManager->SetMenuMultiBox(MenuBarBuilder.GetMultiBox(), MenuWidget);
 
-	const FName SummaryTabId = "SummaryLogTab";
-	const FName ValidatorLogTabId = "ValidatorLogTab";
-	const FName PresubmitLogTabId = "PresubmitLogTab";
-
 	// ----------------------------------------------------------------------
 	// Summary Log Tab
 	// ----------------------------------------------------------------------
-	FOutputLogCreationParams SummaryLogCreationParams;
-
-	SummaryLogCreationParams.AllowAsInitialLogCategory = FAllowLogCategoryCallback::CreateLambda([](const FName LogCategory) {
-		return LogCategory == LogSubmitTool.GetCategoryName();
-		});
-
-	SummaryLogCreationParams.DefaultCategorySelection.Emplace(LogSubmitTool.GetCategoryName(), true);
-	SummaryLogCreationParams.DefaultCategorySelection.Emplace(LogValidators.GetCategoryName(), false);
-	SummaryLogCreationParams.DefaultCategorySelection.Emplace(LogValidatorsResult.GetCategoryName(), true);
-	SummaryLogCreationParams.DefaultCategorySelection.Emplace(LogPresubmitResult.GetCategoryName(), true);
-	SummaryLogCreationParams.DefaultCategorySelection.Emplace(LogOutputDevice.GetCategoryName(), true);
-	SummaryLogCreationParams.DefaultCategorySelection.Emplace(LogSubmitToolP4.GetCategoryName(), true);
-	SummaryLogCreationParams.DefaultCategorySelection.Emplace(LogPresubmit.GetCategoryName(), false);
+	FDefaultCategorySelectionMap DefaultCategorySelection;
+	DefaultCategorySelection.Emplace(LogSubmitTool.GetCategoryName(), true);
+	DefaultCategorySelection.Emplace(LogValidators.GetCategoryName(), false);
+	DefaultCategorySelection.Emplace(LogValidatorsResult.GetCategoryName(), true);
+	DefaultCategorySelection.Emplace(LogPresubmitResult.GetCategoryName(), true);
+	DefaultCategorySelection.Emplace(LogOutputDevice.GetCategoryName(), true);
+	DefaultCategorySelection.Emplace(LogSubmitToolP4.GetCategoryName(), true);
+	DefaultCategorySelection.Emplace(LogPresubmit.GetCategoryName(), false);
 
 	SummaryLogDockTab = SNew(SDockTab)
 		.CanEverClose(false)
@@ -750,7 +744,10 @@ TSharedRef<SWidget> SubmitToolWidget::BuildOutputLogWidget()
 		.Label(FText::FromString(TEXT("Summary")))
 		.OnCanCloseTab(SDockTab::FCanCloseTab::CreateLambda([] { return false; }))
 		[
-			OutputLogModule.MakeOutputLogWidget(SummaryLogCreationParams)
+			SAssignNew(SummaryLogWidget, SSubmitToolOutputLog, false)
+				.Messages(&InLogHistory->GetMessages())
+				.DefaultCategorySelection(DefaultCategorySelection)
+				.AllowInitialLogCategory(FAllowLogCategoryCallback::CreateLambda([](const FName LogCategory) { return LogCategory == LogSubmitTool.GetCategoryName(); }))
 		];
 	
 	LogTabManager->RegisterTabSpawner(SummaryTabId, FOnSpawnTab::CreateLambda([this](const FSpawnTabArgs& SpawnArgs) {	return SummaryLogDockTab.ToSharedRef(); }));
@@ -758,18 +755,14 @@ TSharedRef<SWidget> SubmitToolWidget::BuildOutputLogWidget()
 	// ----------------------------------------------------------------------
 	// Validators Log Tab
 	// ----------------------------------------------------------------------
-	FOutputLogCreationParams ValidatorLogCreationParams;
-	ValidatorLogCreationParams.AllowAsInitialLogCategory = FAllowLogCategoryCallback::CreateLambda([](const FName LogCategory) {
-		return LogCategory == LogValidators.GetCategoryName();
-		});
-
-	ValidatorLogCreationParams.DefaultCategorySelection.Emplace(LogSubmitTool.GetCategoryName(), false);
-	ValidatorLogCreationParams.DefaultCategorySelection.Emplace(LogValidators.GetCategoryName(), true);
-	ValidatorLogCreationParams.DefaultCategorySelection.Emplace(LogValidatorsResult.GetCategoryName(), false);
-	ValidatorLogCreationParams.DefaultCategorySelection.Emplace(LogPresubmitResult.GetCategoryName(), false);
-	ValidatorLogCreationParams.DefaultCategorySelection.Emplace(LogOutputDevice.GetCategoryName(), false);
-	ValidatorLogCreationParams.DefaultCategorySelection.Emplace(LogSubmitToolP4.GetCategoryName(), false);
-	ValidatorLogCreationParams.DefaultCategorySelection.Emplace(LogPresubmit.GetCategoryName(), false);
+	DefaultCategorySelection.Reset();
+	DefaultCategorySelection.Emplace(LogSubmitTool.GetCategoryName(), false);
+	DefaultCategorySelection.Emplace(LogValidators.GetCategoryName(), true);
+	DefaultCategorySelection.Emplace(LogValidatorsResult.GetCategoryName(), false);
+	DefaultCategorySelection.Emplace(LogPresubmitResult.GetCategoryName(), false);
+	DefaultCategorySelection.Emplace(LogOutputDevice.GetCategoryName(), false);
+	DefaultCategorySelection.Emplace(LogSubmitToolP4.GetCategoryName(), false);
+	DefaultCategorySelection.Emplace(LogPresubmit.GetCategoryName(), false);
 
 
 	ValidatorLogTab = SNew(SDockTab)
@@ -778,7 +771,10 @@ TSharedRef<SWidget> SubmitToolWidget::BuildOutputLogWidget()
 		.TabRole(ETabRole::PanelTab)
 		.OnCanCloseTab(SDockTab::FCanCloseTab::CreateLambda([] { return false; }))
 		[
-			OutputLogModule.MakeOutputLogWidget(ValidatorLogCreationParams)
+			SAssignNew(ValidatorLogWidget, SSubmitToolOutputLog, false)
+				.Messages(&InLogHistory->GetMessages())
+				.DefaultCategorySelection(DefaultCategorySelection)
+				.AllowInitialLogCategory(FAllowLogCategoryCallback::CreateLambda([](const FName LogCategory) {	return LogCategory == LogValidators.GetCategoryName(); }))
 		];
 
 	LogTabManager->RegisterTabSpawner(ValidatorLogTabId, FOnSpawnTab::CreateLambda([this](const FSpawnTabArgs& SpawnArgs) { return ValidatorLogTab.ToSharedRef(); }));
@@ -786,18 +782,14 @@ TSharedRef<SWidget> SubmitToolWidget::BuildOutputLogWidget()
 	// ----------------------------------------------------------------------
 	// Pre Submit Log Tab
 	// ----------------------------------------------------------------------
-	FOutputLogCreationParams PresubmitLogCreationParams;
-	PresubmitLogCreationParams.AllowAsInitialLogCategory = FAllowLogCategoryCallback::CreateLambda([](const FName LogCategory) {
-		return LogCategory == LogPresubmit.GetCategoryName();
-		});
-
-	PresubmitLogCreationParams.DefaultCategorySelection.Emplace(LogPresubmit.GetCategoryName(), true);
-	PresubmitLogCreationParams.DefaultCategorySelection.Emplace(LogSubmitTool.GetCategoryName(), false);
-	PresubmitLogCreationParams.DefaultCategorySelection.Emplace(LogValidators.GetCategoryName(), false);
-	PresubmitLogCreationParams.DefaultCategorySelection.Emplace(LogValidatorsResult.GetCategoryName(), false);
-	PresubmitLogCreationParams.DefaultCategorySelection.Emplace(LogPresubmitResult.GetCategoryName(), false);
-	PresubmitLogCreationParams.DefaultCategorySelection.Emplace(LogOutputDevice.GetCategoryName(), false);
-	PresubmitLogCreationParams.DefaultCategorySelection.Emplace(LogSubmitToolP4.GetCategoryName(), false);
+	DefaultCategorySelection.Reset();
+	DefaultCategorySelection.Emplace(LogPresubmit.GetCategoryName(), true);
+	DefaultCategorySelection.Emplace(LogSubmitTool.GetCategoryName(), false);
+	DefaultCategorySelection.Emplace(LogValidators.GetCategoryName(), false);
+	DefaultCategorySelection.Emplace(LogValidatorsResult.GetCategoryName(), false);
+	DefaultCategorySelection.Emplace(LogPresubmitResult.GetCategoryName(), false);
+	DefaultCategorySelection.Emplace(LogOutputDevice.GetCategoryName(), false);
+	DefaultCategorySelection.Emplace(LogSubmitToolP4.GetCategoryName(), false);
 
 	PresubmitLogTab = SNew(SDockTab)
 		.CanEverClose(false)
@@ -805,12 +797,14 @@ TSharedRef<SWidget> SubmitToolWidget::BuildOutputLogWidget()
 		.TabRole(ETabRole::PanelTab)
 		.OnCanCloseTab(SDockTab::FCanCloseTab::CreateLambda([] { return false; }))
 		[
-			OutputLogModule.MakeOutputLogWidget(PresubmitLogCreationParams)
+			SAssignNew(PresubmitLogWidget, SSubmitToolOutputLog, false)
+				.Messages(&InLogHistory->GetMessages())
+				.DefaultCategorySelection(DefaultCategorySelection)
+				.AllowInitialLogCategory(FAllowLogCategoryCallback::CreateLambda([](const FName LogCategory) {	return LogCategory == LogPresubmit.GetCategoryName(); }))
 		];
 
 	LogTabManager->RegisterTabSpawner(PresubmitLogTabId, FOnSpawnTab::CreateLambda([this](const FSpawnTabArgs& SpawnArgs) { return PresubmitLogTab.ToSharedRef(); }));
 
-
 	// ----------------------------------------------------------------------
 	// Logs Layout
 	// ----------------------------------------------------------------------
@@ -885,6 +879,8 @@ TSharedRef<SExpandableArea> SubmitToolWidget::BuildFilesInCLWidget()
 			});
 
 	ModelInterface->FileRefreshedCallback.AddLambda([FileList]() {FileList->RequestListRefresh(); });
+	ModelInterface->SubmitFinishedCallback.AddLambda([this](bool bSuccess) { LogTabManager->TryInvokeTab(SummaryTabId); SummaryLogWidget->ScrollToEnd(); });
+	ModelInterface->PrepareSubmitCallBack.AddLambda([this]() { LogTabManager->TryInvokeTab(SummaryTabId); SummaryLogWidget->ScrollToEnd(); });
 
 	return SNew(SExpandableArea)
 		.InitiallyCollapsed(!FSubmitToolUserPrefs::Get()->bExpandFilesInCL)
@@ -1022,6 +1018,7 @@ void SubmitToolWidget::OnSingleValidatorFinished(const FValidatorBase& InValidat
 	if(!InValidator.GetHasPassed())
 	{
 		ValidatorLogTab->FlashTab();
+		ValidatorLogWidget->ScrollToEnd();
 	}
 }
 
@@ -1034,6 +1031,7 @@ void SubmitToolWidget::OnValidationUpdated(bool bValid)
 	else
 	{
 		ValidateBtn->SetButtonStyle(&FAppStyle::Get().GetWidgetStyle<FButtonStyle>("PrimaryButton"));
+		ValidatorLogWidget->ScrollToEnd();
 	}
 }
 
diff --git a/Engine/Source/Programs/SubmitTool/Private/View/Widgets/SubmitToolWidget.h b/Engine/Source/Programs/SubmitTool/Private/View/Widgets/SubmitToolWidget.h
--- a/Engine/Source/Programs/SubmitTool/Private/View/Widgets/SubmitToolWidget.h
+++ b/Engine/Source/Programs/SubmitTool/Private/View/Widgets/SubmitToolWidget.h
@@ -16,6 +16,8 @@ class FTabManager;
 class FValidatorBase;
 class SExpandableArea;
 class SIntegrationWidget;
+class FSubmitToolOutputLogHistory;
+class SSubmitToolOutputLog;
 
 class SubmitToolWidget final : public SCompoundWidget
 {
@@ -26,7 +28,7 @@ public:
 		SLATE_ARGUMENT(TSharedPtr<SWindow>, ParentWindow)
 	SLATE_END_ARGS()
 
-	void Construct(const FArguments& InArgs);
+	void Construct(const FArguments& InArgs, TSharedPtr<FSubmitToolOutputLogHistory> InLogHistory);
 	virtual ~SubmitToolWidget() override;
 	
 	FReply SubmitClicked();
@@ -36,12 +38,6 @@ private:
 	void HandleApplicationActivationStateChanged(bool bActive);
 	TWeakPtr<SDockTab> ParentTab;
 	FModelInterface* ModelInterface;
-
-	/**
-	 * @brief Shows a dialog if the user wants to delete their shelf
-	 * @return The button that is pressed
-	 */
-	EDialogFactoryResult ShowDeleteShelveDialog() const;
 	FText GetMainButtonText() const;
 	FText GetValidateButtonText() const;
 	FText GetValidateButtonTooltip() const;
@@ -49,11 +45,20 @@ private:
 	TSharedPtr<FTabManager> LogTabManager;
 	TSharedPtr<SButton> ValidateBtn;
 	TSharedRef<SHorizontalBox> BuildButtonRow();
-	TSharedRef<SWidget> BuildOutputLogWidget();
+	TSharedRef<SWidget> BuildOutputLogWidget(const TSharedPtr<FSubmitToolOutputLogHistory>& InLogHistory);
 	TSharedRef<SExpandableArea> BuildFilesInCLWidget();
+
+
 	TSharedPtr<SDockTab> ValidatorLogTab;
+	const FName ValidatorLogTabId = "ValidatorLogTab";
+	TSharedPtr<SSubmitToolOutputLog> ValidatorLogWidget;
 	TSharedPtr<SDockTab> PresubmitLogTab;
+	const FName PresubmitLogTabId = "PresubmitLogTab";
+	TSharedPtr<SSubmitToolOutputLog> PresubmitLogWidget;
 	TSharedPtr<SDockTab> SummaryLogDockTab;
+	const FName SummaryTabId = "SummaryLogTab";
+	TSharedPtr<SSubmitToolOutputLog> SummaryLogWidget;
+
 	TSharedPtr<SMultiLineEditableTextBox> DescriptionBox;
 	TSharedPtr<SIntegrationWidget> IntegrationWidget;
 
diff --git a/Engine/Source/Programs/SubmitTool/SubmitTool.Build.cs b/Engine/Source/Programs/SubmitTool/SubmitTool.Build.cs
--- a/Engine/Source/Programs/SubmitTool/SubmitTool.Build.cs
+++ b/Engine/Source/Programs/SubmitTool/SubmitTool.Build.cs
@@ -1,5 +1,6 @@
 // Copyright Epic Games, Inc. All Rights Reserved.
 
+using System.IO;
 using UnrealBuildTool;
 public class SubmitTool : ModuleRules
 {
@@ -15,6 +16,8 @@ public class SubmitTool : ModuleRules
 			}
 		);
 
+		PrivateIncludePaths.Add(Path.Combine(EngineDirectory, "Source/Developer/OutputLog/Private"));
+
 		PrivateDependencyModuleNames.AddRange(
 			new string[] {
 				"SubmitToolCore",
