Introduce a ValidatorFactory to decouple validator instantiation from ModelInterface
Replace the hardcoded if-else chain for parsing validators and pre-submit operations with a registry-based factory pattern using REGISTER_VALIDATOR_TYPE macros
Merge ParseValidators and ParsePreSubmitOperations into a single unified ParseTasksConfig method

this should allow auto registration by validators in their own Modules

#Synced-CL 51711301

diff --git a/Engine/Source/Developer/SubmitToolCore/Private/Logic/Validators/ValidatorFactory.cpp b/Engine/Source/Developer/SubmitToolCore/Private/Logic/Validators/ValidatorFactory.cpp
new file mode 100644
--- /dev/null
+++ b/Engine/Source/Developer/SubmitToolCore/Private/Logic/Validators/ValidatorFactory.cpp
@@ -0,0 +1,31 @@
+// Copyright Epic Games, Inc. All Rights Reserved.
+
+#include "Logic/Validators/ValidatorFactory.h"
+
+FValidatorFactory& FValidatorFactory::Get()
+{
+	static FValidatorFactory Instance;
+	return Instance;
+}
+
+void FValidatorFactory::Register(const FString& TypeName, FValidatorConstructor InCreator)
+{
+	FString LowerCase = TypeName.ToLower();
+	ensureAlwaysMsgf(!Registry.Contains(LowerCase),
+		TEXT("Validator type '%s' was registered twice"), *TypeName);
+	Registry.Add(LowerCase, MoveTemp(InCreator));
+}
+
+TSharedPtr<FValidatorBase> FValidatorFactory::Create(
+	const FString& TypeName, const FName& Id,
+	const FSubmitToolParameters& Params,
+	const TSharedRef<FSubmitToolServiceProvider>& ServiceProvider,
+	const FString& Definition) const
+{
+	const FValidatorConstructor* Creator = Registry.Find(TypeName.ToLower());
+	if (!Creator)
+	{
+		return nullptr;
+	}
+	return (*Creator)(Id, Params, ServiceProvider, Definition);
+}
diff --git a/Engine/Source/Developer/SubmitToolCore/Public/Logic/Validators/ValidatorFactory.h b/Engine/Source/Developer/SubmitToolCore/Public/Logic/Validators/ValidatorFactory.h
new file mode 100644
--- /dev/null
+++ b/Engine/Source/Developer/SubmitToolCore/Public/Logic/Validators/ValidatorFactory.h
@@ -0,0 +1,55 @@
+// Copyright Epic Games, Inc. All Rights Reserved.
+#pragma once
+#include "CoreMinimal.h"
+
+class FValidatorBase;
+class FSubmitToolServiceProvider;
+struct FSubmitToolParameters;
+
+class FValidatorFactory
+{
+public:
+	using FValidatorConstructor = TFunction<TSharedRef<FValidatorBase>(
+		const FName&,
+		const FSubmitToolParameters&,
+		const TSharedRef<FSubmitToolServiceProvider>&,
+		const FString&)>;
+
+	static FValidatorFactory& Get();
+
+	void Register(const FString& TypeName, FValidatorConstructor InCreator);
+
+	TSharedPtr<FValidatorBase> Create(
+		const FString& TypeName,
+		const FName& Id,
+		const FSubmitToolParameters& Params,
+		const TSharedRef<FSubmitToolServiceProvider>& ServiceProvider,
+		const FString& Definition) const;
+
+private:
+	TMap<FString, FValidatorConstructor> Registry;
+};
+
+struct FValidatorRegistrar
+{
+	FValidatorRegistrar(const FString& TypeName, FValidatorFactory::FValidatorConstructor Creator)
+	{
+		FValidatorFactory::Get().Register(TypeName, MoveTemp(Creator));
+	}
+};
+
+#define REGISTER_VALIDATOR_TYPE(TypeNameConstant, ValidatorClass)				\
+      static FValidatorRegistrar GRegistrar_##ValidatorClass(					\
+          TypeNameConstant,														\
+          [](const FName& Id, const FSubmitToolParameters& Params,				\
+             const TSharedRef<FSubmitToolServiceProvider>& SP,					\
+             const FString& Def) -> TSharedRef<FValidatorBase>					\
+          { return MakeShared<ValidatorClass>(Id, Params, SP, Def); });
+
+#define REGISTER_VALIDATOR_TYPE_ALIAS(TypeNameConstant, ValidatorClass, Alias)	\
+      static FValidatorRegistrar GRegistrar_##Alias(							\
+          TypeNameConstant,														\
+          [](const FName& Id, const FSubmitToolParameters& Params,				\
+             const TSharedRef<FSubmitToolServiceProvider>& SP,					\
+             const FString& Def) -> TSharedRef<FValidatorBase>					\
+          { return MakeShared<ValidatorClass>(Id, Params, SP, Def); });
\ No newline at end of file
diff --git a/Engine/Source/Programs/SubmitTool/Private/Logic/PreSubmitOperations/VirtualizationOperation.cpp b/Engine/Source/Programs/SubmitTool/Private/Logic/PreSubmitOperations/VirtualizationOperation.cpp
--- a/Engine/Source/Programs/SubmitTool/Private/Logic/PreSubmitOperations/VirtualizationOperation.cpp
+++ b/Engine/Source/Programs/SubmitTool/Private/Logic/PreSubmitOperations/VirtualizationOperation.cpp
@@ -5,10 +5,13 @@
 #include "AnalyticsEventAttribute.h"
 #include "Configuration/Configuration.h"
 #include "Logic/Validators/ValidatorDefinition.h"
+#include "Logic/Validators/ValidatorFactory.h"
 #include "Misc/Paths.h"
 #include "Misc/StringOutputDevice.h"
 #include "Models/ModelInterface.h"
 
+REGISTER_VALIDATOR_TYPE(SubmitToolParseConstants::VirtualizationToolOp, FVirtualizationOperation)
+
 const TCHAR* LexToString(FVirtualizationOperation::EVirtualizationErrorCode ErrorCode)
 {
 	switch (ErrorCode)
diff --git a/Engine/Source/Programs/SubmitTool/Private/Logic/PreSubmitOperations/VirtualizationOperation.h b/Engine/Source/Programs/SubmitTool/Private/Logic/PreSubmitOperations/VirtualizationOperation.h
--- a/Engine/Source/Programs/SubmitTool/Private/Logic/PreSubmitOperations/VirtualizationOperation.h
+++ b/Engine/Source/Programs/SubmitTool/Private/Logic/PreSubmitOperations/VirtualizationOperation.h
@@ -5,6 +5,11 @@
 #include "Logic/ProcessPipes.h"
 #include "Logic/Validators/ValidatorRunExecutable.h"
 
+namespace SubmitToolParseConstants
+{
+	const FString VirtualizationToolOp = TEXT("VirtualizationTool");
+}
+
 class FVirtualizationOperation : public FValidatorRunExecutable
 {
 public:
diff --git a/Engine/Source/Programs/SubmitTool/Private/Logic/Validators/CoreRedirectsValidator.cpp b/Engine/Source/Programs/SubmitTool/Private/Logic/Validators/CoreRedirectsValidator.cpp
--- a/Engine/Source/Programs/SubmitTool/Private/Logic/Validators/CoreRedirectsValidator.cpp
+++ b/Engine/Source/Programs/SubmitTool/Private/Logic/Validators/CoreRedirectsValidator.cpp
@@ -7,6 +7,9 @@
 #include "Configuration/Configuration.h"
 #include "Logging/SubmitToolLog.h"
 #include "Logic/Services/Interfaces/ISTSourceControlService.h"
+#include "Logic/Validators/ValidatorFactory.h"
+
+REGISTER_VALIDATOR_TYPE(SubmitToolParseConstants::CoreRedirectsValidator, FCoreRedirectsValidator)
 
 void FCoreRedirectsValidator::StartAsyncWork(const FString& CLDescription, const TArray<FSCFileRef>& FilteredFilesInCL, const TArray<const FTag*>& Tags) 
 {
diff --git a/Engine/Source/Programs/SubmitTool/Private/Logic/Validators/CoreRedirectsValidator.h b/Engine/Source/Programs/SubmitTool/Private/Logic/Validators/CoreRedirectsValidator.h
--- a/Engine/Source/Programs/SubmitTool/Private/Logic/Validators/CoreRedirectsValidator.h
+++ b/Engine/Source/Programs/SubmitTool/Private/Logic/Validators/CoreRedirectsValidator.h
@@ -5,6 +5,11 @@
 #include "ValidatorBaseAsync.h"
 #include "Tasks/Task.h"
 
+namespace SubmitToolParseConstants
+{
+	const FString CoreRedirectsValidator = TEXT("CoreRedirectsValidator");
+}
+
 class FCoreRedirectsValidator : public FValidatorBaseAsync
 {
 public:
diff --git a/Engine/Source/Programs/SubmitTool/Private/Logic/Validators/CrossChangelistValidator.cpp b/Engine/Source/Programs/SubmitTool/Private/Logic/Validators/CrossChangelistValidator.cpp
--- a/Engine/Source/Programs/SubmitTool/Private/Logic/Validators/CrossChangelistValidator.cpp
+++ b/Engine/Source/Programs/SubmitTool/Private/Logic/Validators/CrossChangelistValidator.cpp
@@ -10,6 +10,9 @@
 #include "SubmitToolCoreUtils.h"
 #include "HAL/FileManager.h"
 #include "Misc/FileHelper.h"
+#include "Logic/Validators/ValidatorFactory.h"
+
+REGISTER_VALIDATOR_TYPE(SubmitToolParseConstants::CrossChangelistValidator, FCrossChangelistValidator)
 
 bool FCrossChangelistValidator::Validate(const FString& InCLDescription, const TArray<FSCFileRef>& InFilteredFilesInCL, const TArray<const FTag*>& InTags)
 {
diff --git a/Engine/Source/Programs/SubmitTool/Private/Logic/Validators/CrossChangelistValidator.h b/Engine/Source/Programs/SubmitTool/Private/Logic/Validators/CrossChangelistValidator.h
--- a/Engine/Source/Programs/SubmitTool/Private/Logic/Validators/CrossChangelistValidator.h
+++ b/Engine/Source/Programs/SubmitTool/Private/Logic/Validators/CrossChangelistValidator.h
@@ -4,6 +4,11 @@
 
 #include "ValidatorBase.h"
 
+namespace SubmitToolParseConstants
+{
+	const FString CrossChangelistValidator = TEXT("CrossChangelistValidator");
+}
+
 class FCrossChangelistValidator : public FValidatorBase
 {
 public:
diff --git a/Engine/Source/Programs/SubmitTool/Private/Logic/Validators/EditorCommandletValidator.cpp b/Engine/Source/Programs/SubmitTool/Private/Logic/Validators/EditorCommandletValidator.cpp
--- a/Engine/Source/Programs/SubmitTool/Private/Logic/Validators/EditorCommandletValidator.cpp
+++ b/Engine/Source/Programs/SubmitTool/Private/Logic/Validators/EditorCommandletValidator.cpp
@@ -11,12 +11,14 @@
 
 #include "Configuration/Configuration.h"
 #include "Logging/SubmitToolLog.h"
+#include "Logic/Validators/ValidatorFactory.h"
 #include "Models/ModelInterface.h"
 #include "Models/SubmitToolUserPrefs.h"
 #include "SubmitToolCoreUtils.h"
 
 #include "ProjectEditorRecords.h"
 
+REGISTER_VALIDATOR_TYPE(SubmitToolParseConstants::EditorValidator, FEditorCommandletValidator)
 
 FEditorCommandletValidator::FEditorCommandletValidator(const FName& InNameId, const FSubmitToolParameters& InParameters, TSharedRef<FSubmitToolServiceProvider> InServiceProvider, const FString& InDefinition) :
 	FValidatorRunExecutable(InNameId, InParameters, InServiceProvider, InDefinition)
diff --git a/Engine/Source/Programs/SubmitTool/Private/Logic/Validators/EditorCommandletValidator.h b/Engine/Source/Programs/SubmitTool/Private/Logic/Validators/EditorCommandletValidator.h
--- a/Engine/Source/Programs/SubmitTool/Private/Logic/Validators/EditorCommandletValidator.h
+++ b/Engine/Source/Programs/SubmitTool/Private/Logic/Validators/EditorCommandletValidator.h
@@ -10,6 +10,11 @@ struct FEditorParameters
 	FString EditorArguments;
 };
 
+namespace SubmitToolParseConstants
+{
+	const FString EditorValidator = TEXT("EditorValidator");
+}
+
 class FEditorCommandletValidator : public FValidatorRunExecutable
 {
 
@@ -17,7 +22,7 @@ public:
 	FEditorCommandletValidator(const FName& InNameId, const FSubmitToolParameters& InParameters, TSharedRef<FSubmitToolServiceProvider> InServiceProvider, const FString& InDefinition);
 	virtual void ParseDefinition(const FString& InDefinition) override;
 	virtual bool Validate(const FString& InCLDescription, const TArray<FSCFileRef>& InFilteredFilesInCL, const TArray<const FTag*>& InTags) override;
-	virtual const FString& GetValidatorTypeName() const override { return SubmitToolParseConstants::UBTValidator; }
+	virtual const FString& GetValidatorTypeName() const override { return SubmitToolParseConstants::EditorValidator; }
 	virtual bool Activate() override;
 
 protected:
diff --git a/Engine/Source/Programs/SubmitTool/Private/Logic/Validators/JsonValidator.cpp b/Engine/Source/Programs/SubmitTool/Private/Logic/Validators/JsonValidator.cpp
--- a/Engine/Source/Programs/SubmitTool/Private/Logic/Validators/JsonValidator.cpp
+++ b/Engine/Source/Programs/SubmitTool/Private/Logic/Validators/JsonValidator.cpp
@@ -5,12 +5,15 @@
 #include "Internationalization/Regex.h"
 #include "Serialization/JsonReader.h"
 #include "Serialization/JsonSerializer.h"
+#include "Logic/Validators/ValidatorFactory.h"
 #include "Misc/FileHelper.h"
 #include "Misc/Paths.h"
 #include "Misc/StringOutputDevice.h"
 
 #include "Models/ModelInterface.h"
 
+REGISTER_VALIDATOR_TYPE(SubmitToolParseConstants::JsonValidator, FJsonValidator)
+
 FJsonValidator::FJsonValidator(const FName& InNameId, const FSubmitToolParameters& InParameters, TSharedRef<FSubmitToolServiceProvider> InServiceProvider, const FString& InDefinition) :
 	FValidatorBaseAsync(InNameId, InParameters, InServiceProvider, InDefinition)
 {
diff --git a/Engine/Source/Programs/SubmitTool/Private/Logic/Validators/JsonValidator.h b/Engine/Source/Programs/SubmitTool/Private/Logic/Validators/JsonValidator.h
--- a/Engine/Source/Programs/SubmitTool/Private/Logic/Validators/JsonValidator.h
+++ b/Engine/Source/Programs/SubmitTool/Private/Logic/Validators/JsonValidator.h
@@ -5,6 +5,11 @@
 #include "ValidatorBaseAsync.h"
 #include "Tasks/Task.h"
 
+namespace SubmitToolParseConstants
+{
+	const FString JsonValidator = TEXT("JsonValidator");
+}
+
 class FJsonValidator final : public FValidatorBaseAsync
 {
 public:
diff --git a/Engine/Source/Programs/SubmitTool/Private/Logic/Validators/PackageDataValidator.cpp b/Engine/Source/Programs/SubmitTool/Private/Logic/Validators/PackageDataValidator.cpp
--- a/Engine/Source/Programs/SubmitTool/Private/Logic/Validators/PackageDataValidator.cpp
+++ b/Engine/Source/Programs/SubmitTool/Private/Logic/Validators/PackageDataValidator.cpp
@@ -2,11 +2,14 @@
 
 #include "PackageDataValidator.h"
 
+#include "Logic/Validators/ValidatorFactory.h"
 #include "Misc/StringOutputDevice.h"
 #include "Models/ModelInterface.h"
 #include "UObject/PackageTrailer.h"
 #include "VirtualizationUtilities.h"
 
+REGISTER_VALIDATOR_TYPE(SubmitToolParseConstants::PackageDataValidator, FPackageDataValidator)
+
 FPackageDataValidator::FPackageDataValidator(const FName& InNameId, const FSubmitToolParameters& InParameters, TSharedRef<FSubmitToolServiceProvider> InServiceProvider, const FString& InDefinition)
 	: FValidatorBaseAsync(InNameId, InParameters, InServiceProvider, InDefinition)
 {
diff --git a/Engine/Source/Programs/SubmitTool/Private/Logic/Validators/PackageDataValidator.h b/Engine/Source/Programs/SubmitTool/Private/Logic/Validators/PackageDataValidator.h
--- a/Engine/Source/Programs/SubmitTool/Private/Logic/Validators/PackageDataValidator.h
+++ b/Engine/Source/Programs/SubmitTool/Private/Logic/Validators/PackageDataValidator.h
@@ -4,6 +4,11 @@
 
 #include "ValidatorBaseAsync.h"
 
+namespace SubmitToolParseConstants
+{
+	const FString PackageDataValidator = TEXT("PackageDataValidator");
+}
+
 class FPackageDataValidator final : public FValidatorBaseAsync
 {
 public:
diff --git a/Engine/Source/Programs/SubmitTool/Private/Logic/Validators/PreflightValidator.cpp b/Engine/Source/Programs/SubmitTool/Private/Logic/Validators/PreflightValidator.cpp
--- a/Engine/Source/Programs/SubmitTool/Private/Logic/Validators/PreflightValidator.cpp
+++ b/Engine/Source/Programs/SubmitTool/Private/Logic/Validators/PreflightValidator.cpp
@@ -4,6 +4,7 @@
 
 #include "Logic/PreflightService.h"
 #include "Logic/TagService.h"
+#include "Logic/Validators/ValidatorFactory.h"
 #include "Logging/SubmitToolLog.h"
 #include "Internationalization/Regex.h"
 #include "Models/ModelInterface.h"
@@ -11,6 +12,8 @@
 #include "HAL/PlatformFileManager.h"
 #include "GenericPlatform/GenericPlatformFile.h"
 
+REGISTER_VALIDATOR_TYPE(SubmitToolParseConstants::PreflightValidator, FPreflightValidator)
+
 FPreflightValidator::FPreflightValidator(const FName& InNameId, const FSubmitToolParameters& InParameters, TSharedRef<FSubmitToolServiceProvider> InServiceProvider, const FString& InDefinition)
 : FValidatorBase(InNameId, InParameters, InServiceProvider, InDefinition)
 {
diff --git a/Engine/Source/Programs/SubmitTool/Private/Logic/Validators/PreflightValidator.h b/Engine/Source/Programs/SubmitTool/Private/Logic/Validators/PreflightValidator.h
--- a/Engine/Source/Programs/SubmitTool/Private/Logic/Validators/PreflightValidator.h
+++ b/Engine/Source/Programs/SubmitTool/Private/Logic/Validators/PreflightValidator.h
@@ -7,6 +7,11 @@
 
 class FPreflightService;
 
+namespace SubmitToolParseConstants
+{
+	const FString PreflightValidator = TEXT("PreflightValidator");
+}
+
 class FPreflightValidator : public FValidatorBase
 {
 public:
diff --git a/Engine/Source/Programs/SubmitTool/Private/Logic/Validators/ShaderValidator.h b/Engine/Source/Programs/SubmitTool/Private/Logic/Validators/ShaderValidator.h
--- a/Engine/Source/Programs/SubmitTool/Private/Logic/Validators/ShaderValidator.h
+++ b/Engine/Source/Programs/SubmitTool/Private/Logic/Validators/ShaderValidator.h
@@ -4,6 +4,11 @@
 
 #include "ValidatorBaseAsync.h"
 
+namespace SubmitToolParseConstants
+{
+	const FString ShaderValidator = TEXT("ShaderValidator");
+}
+
 class FShaderValidator final : public FValidatorBaseAsync
 {
 public:		
diff --git a/Engine/Source/Programs/SubmitTool/Private/Logic/Validators/TagValidator.cpp b/Engine/Source/Programs/SubmitTool/Private/Logic/Validators/TagValidator.cpp
--- a/Engine/Source/Programs/SubmitTool/Private/Logic/Validators/TagValidator.cpp
+++ b/Engine/Source/Programs/SubmitTool/Private/Logic/Validators/TagValidator.cpp
@@ -5,8 +5,11 @@
 #include "Logging/SubmitToolLog.h"
 #include "Logic/Services/Interfaces/ISTSourceControlService.h"
 #include "Logic/ChangelistService.h"
+#include "Logic/Validators/ValidatorFactory.h"
 #include "Internationalization/Regex.h"
 
+REGISTER_VALIDATOR_TYPE(SubmitToolParseConstants::TagValidator, FTagValidator)
+
 bool FTagValidator::Validate(const FString& InCLDescription, const TArray<FSCFileRef>& InFilteredFilesInCL, const TArray<const FTag*>& InTags)
 {
 	bResult = true;
diff --git a/Engine/Source/Programs/SubmitTool/Private/Logic/Validators/TagValidator.h b/Engine/Source/Programs/SubmitTool/Private/Logic/Validators/TagValidator.h
--- a/Engine/Source/Programs/SubmitTool/Private/Logic/Validators/TagValidator.h
+++ b/Engine/Source/Programs/SubmitTool/Private/Logic/Validators/TagValidator.h
@@ -6,6 +6,11 @@
 struct FUserData;
 class FTag;
 
+namespace SubmitToolParseConstants
+{
+	const FString TagValidator = TEXT("TagValidator");
+}
+
 class FTagValidator : public FValidatorBase
 {
 public:
diff --git a/Engine/Source/Programs/SubmitTool/Private/Logic/Validators/UBTValidator.cpp b/Engine/Source/Programs/SubmitTool/Private/Logic/Validators/UBTValidator.cpp
--- a/Engine/Source/Programs/SubmitTool/Private/Logic/Validators/UBTValidator.cpp
+++ b/Engine/Source/Programs/SubmitTool/Private/Logic/Validators/UBTValidator.cpp
@@ -9,10 +9,12 @@
 #include "CommandLine/CmdLineParameters.h"
 
 #include "Configuration/Configuration.h"
+#include "Logic/Validators/ValidatorFactory.h"
 #include "Logging/SubmitToolLog.h"
 #include "Models/ModelInterface.h"
 #include "Models/SubmitToolUserPrefs.h"
 
+REGISTER_VALIDATOR_TYPE(SubmitToolParseConstants::UBTValidator, FUBTValidator)
 
 FUBTValidator::FUBTValidator(const FName& InNameId, const FSubmitToolParameters& InParameters, TSharedRef<FSubmitToolServiceProvider> InServiceProvider, const FString& InDefinition) :
 	FValidatorRunExecutable(InNameId, InParameters, InServiceProvider, InDefinition)
diff --git a/Engine/Source/Programs/SubmitTool/Private/Logic/Validators/UBTValidator.h b/Engine/Source/Programs/SubmitTool/Private/Logic/Validators/UBTValidator.h
--- a/Engine/Source/Programs/SubmitTool/Private/Logic/Validators/UBTValidator.h
+++ b/Engine/Source/Programs/SubmitTool/Private/Logic/Validators/UBTValidator.h
@@ -4,6 +4,11 @@
 
 #include "ValidatorRunExecutable.h"
 
+namespace SubmitToolParseConstants
+{
+	const FString UBTValidator = TEXT("UBTValidator");
+}
+
 class FUBTValidator : public FValidatorRunExecutable
 {
 
diff --git a/Engine/Source/Programs/SubmitTool/Private/Logic/Validators/ValidatorBase.h b/Engine/Source/Programs/SubmitTool/Private/Logic/Validators/ValidatorBase.h
--- a/Engine/Source/Programs/SubmitTool/Private/Logic/Validators/ValidatorBase.h
+++ b/Engine/Source/Programs/SubmitTool/Private/Logic/Validators/ValidatorBase.h
@@ -19,22 +19,6 @@
 class FTag;
 struct FAnalyticsEventAttribute;
 
-namespace SubmitToolParseConstants
-{
-	const FString TagValidator = TEXT("TagValidator");
-	const FString UBTValidator = TEXT("UBTValidator");
-	const FString EditorValidator = TEXT("EditorValidator");
-	const FString ShaderValidator = TEXT("ShaderValidator");
-	const FString CustomValidator = TEXT("CustomValidator");
-	const FString CrossChangelistValidator = TEXT("CrossChangelistValidator");
-	const FString PreflightValidator = TEXT("PreflightValidator");
-	const FString PackageDataValidator = TEXT("PackageDataValidator");
-	const FString JsonValidator = TEXT("JsonValidator");
-	const FString InvalidateNode = TEXT("Invalidated");
-	const FString VirtualizationToolOp = TEXT("VirtualizationTool");
-	const FString CoreRedirectsValidator = TEXT("CoreRedirectsValidator");
-};
-
 UENUM()
 enum class EValidationStates : uint8
 {
@@ -78,6 +62,7 @@ class FValidatorBase
 {
 public:
 	FValidatorBase() = delete;
+	// if signature changes, need to update FValidatorFactory
 	FValidatorBase(const FName& InNameId, const FSubmitToolParameters& InParameters, const TWeakPtr<FSubmitToolServiceProvider> InServiceProvider, const FString& InDefinition);
 	virtual ~FValidatorBase();
 
diff --git a/Engine/Source/Programs/SubmitTool/Private/Logic/Validators/ValidatorRunExecutable.cpp b/Engine/Source/Programs/SubmitTool/Private/Logic/Validators/ValidatorRunExecutable.cpp
--- a/Engine/Source/Programs/SubmitTool/Private/Logic/Validators/ValidatorRunExecutable.cpp
+++ b/Engine/Source/Programs/SubmitTool/Private/Logic/Validators/ValidatorRunExecutable.cpp
@@ -8,6 +8,7 @@
 #include "Configuration/Configuration.h"
 #include "HAL/FileManagerGeneric.h"
 #include "Internationalization/Regex.h"
+#include "Logic/Validators/ValidatorFactory.h"
 #include "Misc/FileHelper.h"
 #include "Misc/PathViews.h"
 #include "Misc/Paths.h"
@@ -21,6 +22,11 @@
 #include "Serialization/JsonWriter.h"
 #include "Telemetry/TelemetryService.h"
 
+// Both "CustomValidator" and "Executable" will resolve to this same type, CustomValidator will eventually be deprecated
+REGISTER_VALIDATOR_TYPE(SubmitToolParseConstants::CustomValidator, FValidatorRunExecutable)
+REGISTER_VALIDATOR_TYPE_ALIAS(SubmitToolParseConstants::ExecutableValidator, FValidatorRunExecutable, ExecutableAlias) 
+
+
 int32 FValidatorRunExecutable::ProcessCounter = 0;
 FString FValidatorRunExecutable::GetProcessName(FStringView ExecutablePath)
 {
diff --git a/Engine/Source/Programs/SubmitTool/Private/Logic/Validators/ValidatorRunExecutable.h b/Engine/Source/Programs/SubmitTool/Private/Logic/Validators/ValidatorRunExecutable.h
--- a/Engine/Source/Programs/SubmitTool/Private/Logic/Validators/ValidatorRunExecutable.h
+++ b/Engine/Source/Programs/SubmitTool/Private/Logic/Validators/ValidatorRunExecutable.h
@@ -16,6 +16,12 @@ struct FProcessTrackingData
 	TArray<FString> WarningList;
 };
 
+namespace SubmitToolParseConstants
+{
+	const FString CustomValidator = TEXT("CustomValidator");
+	const FString ExecutableValidator = TEXT("Executable"); // Same as Custom Validator, renaming to keep in line with class
+}
+
 /**
  * @brief Runs an executable with optional arguments. Return code 0 on the executable means success.
  * Validator values: "ExecutablePath" (required), "ExecutableArguments" (optional)
@@ -28,7 +34,7 @@ public:
 
 	virtual void StartValidation() override;
 	virtual bool Validate(const FString& InCLDescription, const TArray<FSCFileRef>& InFilteredFilesInCL, const TArray<const FTag*>& InTags) override;
-	virtual const FString& GetValidatorTypeName() const override { return SubmitToolParseConstants::CustomValidator; }
+	virtual const FString& GetValidatorTypeName() const override { return SubmitToolParseConstants::ExecutableValidator; }
 	virtual bool Activate() override;
 
 	virtual void StopInternalValidations() override;
diff --git a/Engine/Source/Programs/SubmitTool/Private/Models/ModelInterface.cpp b/Engine/Source/Programs/SubmitTool/Private/Models/ModelInterface.cpp
--- a/Engine/Source/Programs/SubmitTool/Private/Models/ModelInterface.cpp
+++ b/Engine/Source/Programs/SubmitTool/Private/Models/ModelInterface.cpp
@@ -22,18 +22,7 @@
 #include "Serialization/JsonSerializer.h"
 
 #include "CommandLine/CmdLineParameters.h"
-#include "Logic/Validators/CoreRedirectsValidator.h"
-#include "Logic/Validators/PreflightValidator.h"
-#include "Logic/Validators/TagValidator.h"
-#include "Logic/Validators/UBTValidator.h"
-#include "Logic/Validators/EditorCommandletValidator.h"
-#include "Logic/Validators/ValidatorBase.h"
-#include "Logic/Validators/ValidatorRunExecutable.h"
-#include "Logic/Validators/CrossChangelistValidator.h"
-#include "Logic/Validators/PackageDataValidator.h"
-#include "Logic/Validators/ShaderValidator.h"
-#include "Logic/Validators/JsonValidator.h"
-#include "Logic/PreSubmitOperations/VirtualizationOperation.h"
+#include "Logic/Validators/ValidatorFactory.h"
 #include "Logic/Services/CacheDataService.h"
 #include "Logic/Services/SourceControl/SubmitToolPerforce.h"
 #include "SubmitToolCoreUtils.h"
@@ -90,8 +79,8 @@ FModelInterface::FModelInterface(FSubmitToolParameters&& InParameters, const TSh
 	UpdateService = MakeShared<FUpdateService>(Parameters.HordeParameters, Parameters.AutoUpdateParameters, ServiceProvider);
 	ServiceProvider->RegisterService<FUpdateService>(UpdateService.ToSharedRef());
 
-	ParseValidators();
-	ParsePreSubmitOperations();
+	ParseTasksConfig(Parameters.Validators, ValidationService.ToSharedRef());
+	ParseTasksConfig(Parameters.PresubmitOperations, PresubmitOperationsService.ToSharedRef());
 
 	OnValidationStateUpdatedHandle = ValidationService->OnTasksRunResultUpdated.Add(FOnTaskRunStateChanged::FDelegate::CreateLambda([this](bool bIsValid) {
 		ReevaluateSubmitToolTag();
@@ -224,12 +213,12 @@ void FModelInterface::Dispose() const
 	ValidationService->StopTasks();
 }
 
-void FModelInterface::ParseValidators()
+void FModelInterface::ParseTasksConfig(const TMap<FString, FString>& InTasksConfig, TSharedRef<FTasksService> InOutTaskService)
 {
-	TRACE_CPUPROFILER_EVENT_SCOPE(FModelInterface::ParseValidators);
+	TRACE_CPUPROFILER_EVENT_SCOPE(FModelInterface::ParseTasksConfig);
 	TArray<TSharedRef<FValidatorBase>> Tasks;
 
-	for (const TPair<FName, FString> DefinitionPair : Parameters.Validators)
+	for (const TPair<FString, FString>& DefinitionPair : InTasksConfig)
 	{
 		FValidatorDefinition TaskDefinition;
 		FStringOutputDevice Errors;
@@ -244,109 +233,29 @@ void FModelInterface::ParseValidators()
 
 		if (TaskDefinition.Type.TrimStartAndEnd().IsEmpty())
 		{
-			UE_LOG(LogSubmitTool, Error, TEXT("Task %s didn't have a Type."), *DefinitionPair.Key.ToString());
+			UE_LOG(LogSubmitTool, Error, TEXT("Task %s didn't have a Type."), *DefinitionPair.Key);
 			continue;
 		}
 
 		if (TaskDefinition.bIsDisabled)
 		{
-			UE_LOG(LogSubmitToolDebug, Log, TEXT("Task %s was disabled by configuration"), *DefinitionPair.Key.ToString());
+			UE_LOG(LogSubmitToolDebug, Log, TEXT("Task %s was disabled by configuration"), *DefinitionPair.Key);
 			continue;
 		}
 
-		if (TaskDefinition.Type.Equals(SubmitToolParseConstants::TagValidator, ESearchCase::IgnoreCase))
-		{
-			Tasks.Add(MakeShared<FTagValidator>(DefinitionPair.Key, Parameters, ServiceProvider.ToSharedRef(), DefinitionPair.Value));
-		}
-		else if (TaskDefinition.Type.Equals(SubmitToolParseConstants::UBTValidator, ESearchCase::IgnoreCase))
-		{
-			Tasks.Add(MakeShared<FUBTValidator>(DefinitionPair.Key, Parameters, ServiceProvider.ToSharedRef(), DefinitionPair.Value));
-		}
-		else if (TaskDefinition.Type.Equals(SubmitToolParseConstants::EditorValidator, ESearchCase::IgnoreCase))
-		{
-			Tasks.Add(MakeShared<FEditorCommandletValidator>(DefinitionPair.Key, Parameters, ServiceProvider.ToSharedRef(), DefinitionPair.Value));
-		}
-		else if (TaskDefinition.Type.Equals(SubmitToolParseConstants::CustomValidator, ESearchCase::IgnoreCase))
-		{
-			Tasks.Add(MakeShared<FValidatorRunExecutable>(DefinitionPair.Key, Parameters, ServiceProvider.ToSharedRef(), DefinitionPair.Value));
-		}
-		else if (TaskDefinition.Type.Equals(SubmitToolParseConstants::CrossChangelistValidator, ESearchCase::IgnoreCase))
-		{
-			Tasks.Add(MakeShared<FCrossChangelistValidator>(DefinitionPair.Key, Parameters, ServiceProvider.ToSharedRef(), DefinitionPair.Value));
-		}
-		else if (TaskDefinition.Type.Equals(SubmitToolParseConstants::PreflightValidator, ESearchCase::IgnoreCase))
-		{
-			Tasks.Add(MakeShared<FPreflightValidator>(DefinitionPair.Key, Parameters, ServiceProvider.ToSharedRef(), DefinitionPair.Value));
-		}
-		else if (TaskDefinition.Type.Equals(SubmitToolParseConstants::PackageDataValidator, ESearchCase::IgnoreCase))
-		{
-			Tasks.Add(MakeShared<FPackageDataValidator>(DefinitionPair.Key, Parameters, ServiceProvider.ToSharedRef(), DefinitionPair.Value));
-		}
-		else if (TaskDefinition.Type.Equals(SubmitToolParseConstants::ShaderValidator, ESearchCase::IgnoreCase))
-		{
-			Tasks.Add(MakeShared<FShaderValidator>(DefinitionPair.Key, Parameters, ServiceProvider.ToSharedRef(), DefinitionPair.Value));
-		}
-		else if (TaskDefinition.Type.Equals(SubmitToolParseConstants::JsonValidator, ESearchCase::IgnoreCase))
-		{
-			Tasks.Add(MakeShared<FJsonValidator>(DefinitionPair.Key, Parameters, ServiceProvider.ToSharedRef(), DefinitionPair.Value));
-		}
-		else if (TaskDefinition.Type.Equals(SubmitToolParseConstants::CoreRedirectsValidator, ESearchCase::IgnoreCase))
-		{
-			Tasks.Add(MakeShared<FCoreRedirectsValidator>(DefinitionPair.Key, Parameters, ServiceProvider.ToSharedRef(), DefinitionPair.Value));
-		}
-		else
-		{
-			UE_LOG(LogSubmitTool, Error, TEXT("[%s] is not a recognized validator type and has not been activated."), *DefinitionPair.Key.ToString());
-		}
-	}
-
-	for (const TSharedRef<FValidatorBase>& Task : Tasks)
-	{
-		Task->OnLockdownIdsChanged.AddRaw(this, &FModelInterface::DynamicValidatorLockdownRefreshed);
-	}
-
-	ValidationService->InitializeTasks(Tasks);
-}
-
-void FModelInterface::ParsePreSubmitOperations() const
-{
-	TRACE_CPUPROFILER_EVENT_SCOPE(FModelInterface::ParsePreSubmitOperations);
-	TArray<TSharedRef<FValidatorBase>> Tasks;
-
-	for (const TPair<FName, FString> DefinitionPair : Parameters.PresubmitOperations)
-	{
-		FValidatorDefinition TaskDefinition;
-		FStringOutputDevice Errors;
-		FValidatorDefinition::StaticStruct()->ImportText(*DefinitionPair.Value, &TaskDefinition, nullptr, 0, &Errors, FValidatorDefinition::StaticStruct()->GetName());
+		TSharedPtr<FValidatorBase> Task = FValidatorFactory::Get().Create(TaskDefinition.Type, FName(DefinitionPair.Key), Parameters, ServiceProvider.ToSharedRef(), DefinitionPair.Value);
 
-		if (!Errors.IsEmpty())
+		if (Task == nullptr)
 		{
-			UE_LOG(LogSubmitTool, Error, TEXT("Error loading parameter file %s"), *Errors);
-			FModelInterface::SetErrorState();
+			UE_LOG(LogSubmitTool, Error, TEXT("[%s] is not a recognized validator type."), *DefinitionPair.Key);
 			continue;
 		}
 
-		if (TaskDefinition.Type.TrimStartAndEnd().IsEmpty())
-		{
-			UE_LOG(LogSubmitTool, Error, TEXT("Task %s didn't have a Type."), *DefinitionPair.Key.ToString());
-			continue;
-		}
-
-		if (TaskDefinition.Type.Equals(SubmitToolParseConstants::CustomValidator, ESearchCase::IgnoreCase))
-		{
-			Tasks.Add(MakeShared<FValidatorRunExecutable>(DefinitionPair.Key, Parameters, ServiceProvider.ToSharedRef(), DefinitionPair.Value));
-		}
-		else if (TaskDefinition.Type.Equals(SubmitToolParseConstants::VirtualizationToolOp, ESearchCase::IgnoreCase))
-		{
-			Tasks.Add(MakeShared<FVirtualizationOperation>(DefinitionPair.Key, Parameters, ServiceProvider.ToSharedRef(), DefinitionPair.Value));
-		}
-		else
-		{
-			UE_LOG(LogSubmitTool, Error, TEXT("[%s] is not a recognized pre submit operation type and has not been activated."), *DefinitionPair.Key.ToString());
-		}
+		Task->OnLockdownIdsChanged.AddRaw(this, &FModelInterface::DynamicValidatorLockdownRefreshed);
+		Tasks.Add(Task.ToSharedRef());
 	}
 
-	PresubmitOperationsService->InitializeTasks(Tasks);
+	InOutTaskService->InitializeTasks(Tasks);
 }
 
 void FModelInterface::GenerateCLDescription()
diff --git a/Engine/Source/Programs/SubmitTool/Private/Models/ModelInterface.h b/Engine/Source/Programs/SubmitTool/Private/Models/ModelInterface.h
--- a/Engine/Source/Programs/SubmitTool/Private/Models/ModelInterface.h
+++ b/Engine/Source/Programs/SubmitTool/Private/Models/ModelInterface.h
@@ -81,8 +81,7 @@ public:
 	FModelInterface(FSubmitToolParameters&& InParameters, const TSharedPtr<ISTSourceControlService>& InSourceControlService);
 	~FModelInterface();
 	void Dispose() const;
-	void ParseValidators();
-	void ParsePreSubmitOperations() const;
+	void ParseTasksConfig(const TMap<FString, FString>& InTasksConfig, TSharedRef<FTasksService> InOutTaskService);
 
 	void SetMainTab(TSharedPtr<SDockTab> InMainTab)									{ MainTab = InMainTab; }
 	TWeakPtr<SDockTab> GetMainTab() const											{ return MainTab; }
