Switch summary log tab activation from TryInvokeTab to ActivateInParent to fix log output redirection being lost on recreation
Use MoveTemp for log history transfers to avoid unnecessary copies and fix deadlock on removing output device capture from shutdown

#Synced-CL 51711301

diff --git a/Engine/Source/Programs/SubmitTool/Private/SubmitToolApp.cpp b/Engine/Source/Programs/SubmitTool/Private/SubmitToolApp.cpp
--- a/Engine/Source/Programs/SubmitTool/Private/SubmitToolApp.cpp
+++ b/Engine/Source/Programs/SubmitTool/Private/SubmitToolApp.cpp
@@ -96,7 +96,7 @@ void FSlateUI::Build(FSubmitToolUserPrefs* UserPrefs, TSharedPtr<FSubmitToolOutp
 
 	FName TabName = FName("Submit Tool");
 
-	TabManager->RegisterTabSpawner(TabName, FOnSpawnTab::CreateLambda([&Window = Window, InLogHistory](const FSpawnTabArgs& SpawnArgs) mutable { return Window.BuildMainTab(SpawnArgs.GetOwnerWindow(), InLogHistory); }));
+	TabManager->RegisterTabSpawner(TabName, FOnSpawnTab::CreateLambda([&Window = Window, InLogHistory](const FSpawnTabArgs& SpawnArgs) mutable { return Window.BuildMainTab(SpawnArgs.GetOwnerWindow(), MoveTemp(InLogHistory)); }));
 	TabManager->RegisterDefaultTabWindowSize(TabName, FVector2D(1024, 600));
 
 	Tab = TabManager->TryInvokeTab(TabName);
diff --git a/Engine/Source/Programs/SubmitTool/Private/View/SubmitToolWindow.cpp b/Engine/Source/Programs/SubmitTool/Private/View/SubmitToolWindow.cpp
--- a/Engine/Source/Programs/SubmitTool/Private/View/SubmitToolWindow.cpp
+++ b/Engine/Source/Programs/SubmitTool/Private/View/SubmitToolWindow.cpp
@@ -52,7 +52,7 @@ void SubmitToolWindow::CreateAutoUpdateSubmitToolContent(TSharedPtr<SWindow> InP
 	TSharedRef<SAutoUpdateWidget> AutoUpdateWidget =
 		SNew(SAutoUpdateWidget)
 		.ModelInterface(ModelInterface)
-		.OnAutoUpdateCancelled_Lambda([InParentWindow, this, InLogHistory]() mutable { CreateMainSubmitToolContent(InParentWindow, InLogHistory); });
+		.OnAutoUpdateCancelled_Lambda([InParentWindow, this, InLogHistory]() mutable { CreateMainSubmitToolContent(InParentWindow, MoveTemp(InLogHistory)); });
 
 	ModelInterface->SetMainTab(MainTab);
 	MainTab->SetContent(AutoUpdateWidget);
diff --git a/Engine/Source/Programs/SubmitTool/Private/View/Widgets/SubmitToolWidget.cpp b/Engine/Source/Programs/SubmitTool/Private/View/Widgets/SubmitToolWidget.cpp
--- a/Engine/Source/Programs/SubmitTool/Private/View/Widgets/SubmitToolWidget.cpp
+++ b/Engine/Source/Programs/SubmitTool/Private/View/Widgets/SubmitToolWidget.cpp
@@ -879,8 +879,8 @@ TSharedRef<SExpandableArea> SubmitToolWidget::BuildFilesInCLWidget()
 			});
 
 	ModelInterface->FileRefreshedCallback.AddLambda([FileList]() {FileList->RequestListRefresh(); });
-	ModelInterface->SubmitFinishedCallback.AddLambda([this](bool bSuccess) { LogTabManager->TryInvokeTab(SummaryTabId); SummaryLogWidget->ScrollToEnd(); });
-	ModelInterface->PrepareSubmitCallBack.AddLambda([this]() { LogTabManager->TryInvokeTab(SummaryTabId); SummaryLogWidget->ScrollToEnd(); });
+	ModelInterface->SubmitFinishedCallback.AddLambda([this](bool bSuccess) { SummaryLogDockTab->ActivateInParent(ETabActivationCause::SetDirectly); SummaryLogWidget->ScrollToEnd(); });
+	ModelInterface->PrepareSubmitCallBack.AddLambda([this]() { SummaryLogDockTab->ActivateInParent(ETabActivationCause::SetDirectly); SummaryLogWidget->ScrollToEnd(); });
 
 	return SNew(SExpandableArea)
 		.InitiallyCollapsed(!FSubmitToolUserPrefs::Get()->bExpandFilesInCL)
