Fix deadlock caused by log output redirectors by introducing an explicit Release method on the output log history to unregister the device before the UI widget is created
Simplify the log message serialization by removing the message count cap and its associated self-unregistration logic that was causing a different deadlock
Remove unused stream data process fields from PreflightService

#Synced-CL 51711301

diff --git a/Engine/Source/Programs/SubmitTool/Private/Logging/SubmitToolOutputLogHistory.h b/Engine/Source/Programs/SubmitTool/Private/Logging/SubmitToolOutputLogHistory.h
--- a/Engine/Source/Programs/SubmitTool/Private/Logging/SubmitToolOutputLogHistory.h
+++ b/Engine/Source/Programs/SubmitTool/Private/Logging/SubmitToolOutputLogHistory.h
@@ -16,6 +16,11 @@ public:
 	}
 
 	virtual ~FSubmitToolOutputLogHistory()
+	{
+		Release();
+	}
+
+	void Release()
 	{
 		if (GLog != nullptr)
 		{
@@ -34,16 +39,7 @@ public:
 private:
 	virtual void Serialize(const TCHAR* V, ELogVerbosity::Type Verbosity, const FName& Category) override
 	{
-		// this log history is meant to be short lived, its purpose is to store messages before the UI widget is created
-		if (Messages.Num() < 2000)
-		{
-			SOutputLog::CreateLogMessages(V, Verbosity, Category, Messages);
-		}
-		else
-		{
-			GLog->RemoveOutputDevice(this);
-			UE_LOG(LogSubmitToolDebug, Error, TEXT("Output log history went over 2000 messages, unregistering"));
-		}
+		SOutputLog::CreateLogMessages(V, Verbosity, Category, Messages);
 	}
 
 	TArray<TSharedPtr<FOutputLogMessage>> Messages;
diff --git a/Engine/Source/Programs/SubmitTool/Private/Logic/PreflightService.h b/Engine/Source/Programs/SubmitTool/Private/Logic/PreflightService.h
--- a/Engine/Source/Programs/SubmitTool/Private/Logic/PreflightService.h
+++ b/Engine/Source/Programs/SubmitTool/Private/Logic/PreflightService.h
@@ -110,11 +110,6 @@ private:
 	FString StreamName;
 	FString TemplateId;
 
-	// get stream data process
-	TUniquePtr<FProcessWrapper> GetStreamDataProcess;
-	FOnCompleted OnGetStreamDataCompletedCallback;
-	FOnOutputLine OnGetStreamDataOutputLineCallback;
-
 	// Fetch Preflight
 	TSharedPtr<IHttpRequest> LinkedPFRequest = nullptr;
 	TMap<FString, TSharedPtr<IHttpRequest>> UnlinkedPFRequests;
diff --git a/Engine/Source/Programs/SubmitTool/Private/View/Widgets/SubmitToolWidget.cpp b/Engine/Source/Programs/SubmitTool/Private/View/Widgets/SubmitToolWidget.cpp
--- a/Engine/Source/Programs/SubmitTool/Private/View/Widgets/SubmitToolWidget.cpp
+++ b/Engine/Source/Programs/SubmitTool/Private/View/Widgets/SubmitToolWidget.cpp
@@ -726,6 +726,9 @@ TSharedRef<SWidget> SubmitToolWidget::BuildOutputLogWidget(const TSharedPtr<FSub
 	LogTabManager->SetAllowWindowMenuBar(true);
 	LogTabManager->SetMenuMultiBox(MenuBarBuilder.GetMultiBox(), MenuWidget);
 
+	// Release the Outputlog redirect to avoid races
+	InLogHistory->Release();
+
 	// ----------------------------------------------------------------------
 	// Summary Log Tab
 	// ----------------------------------------------------------------------
