Mark shelf as complete across all three success paths in the shelving logic to ensure consistent state tracking

#Synced-CL 52011596

diff --git a/Engine/Source/Programs/SubmitTool/Private/Logic/ChangelistService.cpp b/Engine/Source/Programs/SubmitTool/Private/Logic/ChangelistService.cpp
--- a/Engine/Source/Programs/SubmitTool/Private/Logic/ChangelistService.cpp
+++ b/Engine/Source/Programs/SubmitTool/Private/Logic/ChangelistService.cpp
@@ -491,6 +491,8 @@ void FChangelistService::ReplaceShelvedFiles(TFunction<void(const FSCCommand&)>
 				{
 					CacheService->SetLastShelfKnownUpToDate(GetCLID(), FDateTime::UtcNow());
 				}
+
+				Changelist.bIsShelfComplete = true;
 				bIsShelfReady = true;
 			}
 			else if (Cmd.bDroppedConnection)
@@ -531,6 +533,7 @@ void FChangelistService::CreateShelvedFiles(TFunction<void(const FSCCommand&)> O
 					CacheService->SetLastShelfKnownUpToDate(GetCLID(), FDateTime::UtcNow());
 				}
 
+				Changelist.bIsShelfComplete = true;
 				bIsShelfReady = true;
 			}
 			else if (Cmd.bDroppedConnection)
@@ -570,7 +573,9 @@ void FChangelistService::UpdateShelvedFiles(TArray<FString> InFiles, TFunction<v
 				{
 					CacheService->SetLastShelfKnownUpToDate(GetCLID(), FDateTime::UtcNow());
 				}
+
 				UE_LOG(LogSubmitToolP4, Log, TEXT("CL %s shelved files have been updated"), *GetCLID());
+				Changelist.bIsShelfComplete = true;
 				bIsShelfReady = true;
 			}
 			else if (Cmd.bDroppedConnection)
