Add OAuth refresh token support to the Jira integration in Submit Tool
Automatically refresh expired access tokens using a stored refresh token, and re-authenticate when refresh fails
Also add SubTask to Jira issue query filters and set the OAuth waiting flag before starting the HTTP listener
Fixed missing callbacks for edge cases and minor bugs

#Synced-CL 52011596

diff --git a/Engine/Source/Programs/SubmitTool/Private/Logic/CredentialsService.cpp b/Engine/Source/Programs/SubmitTool/Private/Logic/CredentialsService.cpp
--- a/Engine/Source/Programs/SubmitTool/Private/Logic/CredentialsService.cpp
+++ b/Engine/Source/Programs/SubmitTool/Private/Logic/CredentialsService.cpp
@@ -266,6 +266,7 @@ bool FCredentialsService::ParseOIDCTokenData(const FString& InToken)
 
 constexpr int JiraCredentialDatVersion = 1;
 constexpr int JiraCredentialTokenVersion = 2;
+constexpr int JiraCredentialRefreshTokenVersion = 3;
 void FCredentialsService::SaveCredentials() const
 {
 	TArray<uint8> Bytes;
@@ -293,7 +294,7 @@ void FCredentialsService::SaveCredentials() const
 
 	if(File != nullptr)
 	{
-		int32 Version = JiraCredentialTokenVersion;
+		int32 Version = JiraCredentialRefreshTokenVersion;
 		*File << Version;
 
 		*File << NumBytesEncrypted;
@@ -343,7 +344,7 @@ void FCredentialsService::LoadCredentials()
 				*File << Version;
 
 				// Check Versions here
-				if(Version > JiraCredentialTokenVersion)
+				if(Version > JiraCredentialRefreshTokenVersion)
 				{
 					UE_LOG(LogSubmitToolDebug, Warning, TEXT("Unexpected Credentials Version, aborting credentials loading."));
 					File->Close();
@@ -389,6 +390,10 @@ void FCredentialsService::LoadCredentials()
 						FMemoryReader Reader(DeserializedBytes);
 						FJiraToken Token;
 						FJiraToken::StaticStruct()->SerializeItem(Reader, &Token, nullptr);
+						if (Version < JiraCredentialRefreshTokenVersion)
+						{
+							Token.Refresh_Token = TEXT("");
+						}
 						JiraOAuthToken = MakeUnique<FJiraToken>(MoveTemp(Token));
 					}
 				}
diff --git a/Engine/Source/Programs/SubmitTool/Private/Logic/CredentialsService.h b/Engine/Source/Programs/SubmitTool/Private/Logic/CredentialsService.h
--- a/Engine/Source/Programs/SubmitTool/Private/Logic/CredentialsService.h
+++ b/Engine/Source/Programs/SubmitTool/Private/Logic/CredentialsService.h
@@ -48,6 +48,8 @@ struct FJiraToken
 	FString Access_Token;
 	UPROPERTY()
 	int32 Expires_In;
+	UPROPERTY()
+	FString Refresh_Token;
 };
 
 class FCredentialsService : public ISubmitToolService
diff --git a/Engine/Source/Programs/SubmitTool/Private/Logic/JiraService.cpp b/Engine/Source/Programs/SubmitTool/Private/Logic/JiraService.cpp
--- a/Engine/Source/Programs/SubmitTool/Private/Logic/JiraService.cpp
+++ b/Engine/Source/Programs/SubmitTool/Private/Logic/JiraService.cpp
@@ -49,9 +49,9 @@ void FJiraService::ObtainOAuth(TFunction<void(bool)> InCallback)
 	{
 		if (!HttpServerListener.IsValid() || !HttpServerListener->IsListening())
 		{
+			bWaitingOAuth = true;
 			HttpServerListener = MakeUnique<FHttpServerListener>(TEXT("/submittool/code"), EHttpServerRequestVerbs::VERB_GET, 8451, FHttpRequestHandler::CreateLambda([this, InCallback](const FHttpServerRequest& Request, const FHttpResultCallback& OnComplete)
 			{
-				bWaitingOAuth = true;
 				const FString Html = TEXT(
 					"<html><body>"
 					"<h2>Authentication received</h2>"
@@ -94,11 +94,13 @@ void FJiraService::ObtainOAuth(TFunction<void(bool)> InCallback)
 			if(!HttpServerListener->StartListening())
 			{
 				UE_LOG(LogSubmitTool, Error, TEXT("Failed to set up Jira OAuth listener."));
+				bWaitingOAuth = false;
+				InCallback(false);
 				return;
 			}
 		}
 
-		FString Url = FString::Format(TEXT("https://auth.atlassian.com/authorize?audience=api.atlassian.com&client_id={0}&scope=read%3Ajira-work%20read%3Ajira-user&redirect_uri=http%3A%2F%2F127.0.0.1%3A8451%2Fsubmittool%2Fcode&state={1}&response_type=code&prompt=consent"), { Definition.OAuthClientId, UserId.ToString(EGuidFormats::Digits) });
+		FString Url = FString::Format(TEXT("https://auth.atlassian.com/authorize?audience=api.atlassian.com&client_id={0}&scope=offline_access%20read%3Ajira-work%20read%3Ajira-user&redirect_uri=http%3A%2F%2F127.0.0.1%3A8451%2Fsubmittool%2Fcode&state={1}&response_type=code&prompt=consent"), { Definition.OAuthClientId, UserId.ToString(EGuidFormats::Digits) });
 		FPlatformProcess::LaunchURL(*Url, nullptr, nullptr);
 		UE_LOG(LogSubmitToolDebug, Log, TEXT("Sending OAuth request for code"));
 	}
@@ -124,7 +126,7 @@ void FJiraService::Reset()
 	JiraIssues.Reset();
 }
 
-void FJiraService::QueryIssues(bool bInFailSilently)
+void FJiraService::QueryIssues(bool bInFailSilently, bool bAllowTokenRefresh)
 {
 	TSharedPtr<FCredentialsService> CredentialsService = ServiceProvider.Pin()->GetService<FCredentialsService>();
 	if(JiraRequest.IsValid() || Definition.bDisableJira || !CredentialsService->HasJiraCredentials())
@@ -134,8 +136,8 @@ void FJiraService::QueryIssues(bool bInFailSilently)
 
 	FHttpModule& HttpModule = FModuleManager::GetModuleChecked<FHttpModule>("HTTP");
 	JiraRequest = HttpModule.Get().CreateRequest();
-	
-	JiraRequest->OnProcessRequestComplete().BindLambda([this, bInFailSilently](FHttpRequestPtr HttpRequest, FHttpResponsePtr HttpResponse, bool bSucceeded) { QueryIssues_HttpRequestComplete(HttpRequest, HttpResponse, bSucceeded, bInFailSilently); });
+
+	JiraRequest->OnProcessRequestComplete().BindLambda([this, bInFailSilently, bAllowTokenRefresh](FHttpRequestPtr HttpRequest, FHttpResponsePtr HttpResponse, bool bSucceeded) { QueryIssues_HttpRequestComplete(HttpRequest, HttpResponse, bSucceeded, bInFailSilently, bAllowTokenRefresh); });
 
 	FString Url;
 	if (Definition.OAuthClientId.IsEmpty())
@@ -145,7 +147,7 @@ void FJiraService::QueryIssues(bool bInFailSilently)
 	}
 	else
 	{
-		Url = FString::Format(TEXT("https://api.atlassian.com/ex/jira/{0}/rest/api/2/search/jql?maxResults={1}&jql=assignee%20%3D%20currentUser%28%29%20%20AND%20type%20IN%20(Bug%2C%20Task)%20ORDER%20BY%20created&fields=key,description,summary,priority,status,issuetype,created"), { Definition.CloudID, MaxResults });
+		Url = FString::Format(TEXT("https://api.atlassian.com/ex/jira/{0}/rest/api/2/search/jql?maxResults={1}&jql=assignee%20%3D%20currentUser%28%29%20%20AND%20type%20IN%20(Bug%2C%20Task%2C%20SubTask)%20ORDER%20BY%20created&fields=key,description,summary,priority,status,issuetype,created"), { Definition.CloudID, MaxResults });
 		JiraRequest->SetHeader(TEXT("Authorization"), TEXT("Bearer ") + CredentialsService->GetJiraToken()->Access_Token);
 	}
 
@@ -156,7 +158,7 @@ void FJiraService::QueryIssues(bool bInFailSilently)
 	JiraRequest->ProcessRequest();
 }
 
-void FJiraService::QueryIssues_HttpRequestComplete(FHttpRequestPtr HttpRequest, FHttpResponsePtr HttpResponse, bool bSucceeded, bool bFailSilently)
+void FJiraService::QueryIssues_HttpRequestComplete(FHttpRequestPtr HttpRequest, FHttpResponsePtr HttpResponse, bool bSucceeded, bool bFailSilently, bool bAllowTokenRefresh)
 {
 	JiraRequest = nullptr;
 	bOngoingRequest = false;
@@ -165,6 +167,7 @@ void FJiraService::QueryIssues_HttpRequestComplete(FHttpRequestPtr HttpRequest,
 	if (!bSucceeded)
 	{
 		UE_LOG(LogSubmitToolDebug, Error, TEXT("Unable to retrieve JIRA issues at the moment."));
+		OnJiraIssuesRetrievedCallback.ExecuteIfBound(false);
 		return;
 	}
 
@@ -224,6 +227,32 @@ void FJiraService::QueryIssues_HttpRequestComplete(FHttpRequestPtr HttpRequest,
 		}
 		else
 		{
+			if (HttpResponse->GetResponseCode() == EHttpResponseCodes::Denied && !Definition.OAuthClientId.IsEmpty() && bAllowTokenRefresh)
+			{
+				UE_LOG(LogSubmitToolDebug, Log, TEXT("Jira access token expired (401), attempting token refresh."));
+				RefreshOAuthToken([this, bFailSilently](bool bRefreshSucceeded)
+				{
+					if (bRefreshSucceeded)
+					{
+						QueryIssues(bFailSilently, /*bAllowTokenRefresh=*/false);
+					}
+					else
+					{
+						UE_LOG(LogSubmitTool, Error, TEXT("Jira token refresh failed, re-authentication required."));
+						ServiceProvider.Pin()->GetService<FCredentialsService>()->InvalidateCredentials();
+						OnJiraIssuesRetrievedCallback.ExecuteIfBound(false);
+					}
+				});
+				return;
+			}
+
+			if (HttpResponse->GetResponseCode() == EHttpResponseCodes::TooManyRequests)
+			{
+				UE_LOG(LogSubmitTool, Warning, TEXT("Jira request rate limited (429), please try again later."));
+				OnJiraIssuesRetrievedCallback.ExecuteIfBound(false);
+				return;
+			}
+
 			if (bFailSilently)
 			{
 				UE_LOG(LogSubmitToolDebug, Error, TEXT("Suppressed Log: Jira Request failed with error code %d, please make sure you're logging with the right credentials. if your Okta password expired recently, make sure you log into JIRA via browser at least once."), HttpResponse->GetResponseCode());
@@ -402,7 +431,7 @@ void FJiraService::ExchangeOAuthCode(TFunction<void(bool)> InCallback)
 	RequestJson->SetStringField(TEXT("client_id"), Definition.OAuthClientId);
 	RequestJson->SetStringField(TEXT("client_secret"), Definition.OAuthClientSecret);
 	RequestJson->SetStringField(TEXT("code"), AuthCode);
-	RequestJson->SetStringField(TEXT("redirect_uri"), TEXT("http://127.0.0.1:8452/submittool/token"));	
+	RequestJson->SetStringField(TEXT("redirect_uri"), TEXT("http://127.0.0.1:8452/submittool/token"));
 
 	FString BodyString;
 	TSharedRef<TJsonWriter<TCHAR>> JsonWriter = TJsonWriterFactory<>::Create(&BodyString);
@@ -457,50 +486,79 @@ void FJiraService::ExchangeOAuthCode(TFunction<void(bool)> InCallback)
 	HttpRequest->ProcessRequest();
 }
 
-const FString FJiraService::GetJiraIssuesFilepath() const
+void FJiraService::RefreshOAuthToken(TFunction<void(bool)> InCallback)
 {
-	return FPaths::Combine(FSubmitToolCoreUtils::GetLocalAppDataPath(), TEXT("SubmitTool"), TEXT("jira.issues.dat"));
-}
+	TSharedPtr<FCredentialsService> CredentialsService = ServiceProvider.Pin()->GetService<FCredentialsService>();
+	const TUniquePtr<FJiraToken>& Token = CredentialsService->GetJiraToken();
+
+	if (!Token.IsValid() || Token->Refresh_Token.IsEmpty())
+	{
+		UE_LOG(LogSubmitTool, Warning, TEXT("No Jira refresh token available, re-authentication required."));
+		InCallback(false);
+		return;
+	}
 
-void FJiraService::GetuserInfo()
-{
 	FHttpModule& HttpModule = FModuleManager::LoadModuleChecked<FHttpModule>("HTTP");
 	FHttpRequestRef HttpRequest = HttpModule.Get().CreateRequest();
 
-	HttpRequest->OnProcessRequestComplete().BindRaw(this, &FJiraService::GetuserInfo_HttpRequestComplete);
-
-	FString Url = FString::Format(TEXT("https://{0}/rest/api/2/myself"), { Definition.ServerAddress });
-
-	HttpRequest->SetURL(Url);
+	HttpRequest->SetURL(TEXT("https://auth.atlassian.com/oauth/token"));
+	HttpRequest->SetHeader(TEXT("Content-Type"), TEXT("application/json"));
+	HttpRequest->SetVerb(TEXT("POST"));
 
-	HttpRequest->SetHeader(TEXT("Authorization"), TEXT("Basic ") + ServiceProvider.Pin()->GetService<FCredentialsService>()->GetEncodedLoginString());
-	HttpRequest->SetVerb(TEXT("GET"));
-	HttpRequest->ProcessRequest();
-}
+	TSharedPtr<FJsonObject> RequestJson = MakeShared<FJsonObject>();
+	RequestJson->SetStringField(TEXT("grant_type"), TEXT("refresh_token"));
+	RequestJson->SetStringField(TEXT("client_id"), Definition.OAuthClientId);
+	RequestJson->SetStringField(TEXT("client_secret"), Definition.OAuthClientSecret);
+	RequestJson->SetStringField(TEXT("refresh_token"), Token->Refresh_Token);
 
-void FJiraService::GetuserInfo_HttpRequestComplete(FHttpRequestPtr HttpRequest, FHttpResponsePtr HttpResponse, bool bSucceeded)
-{
-	if (!bSucceeded)
-	{
-		UE_LOG(LogSubmitToolDebug, Error, TEXT("Unable to retrieve JIRA issues at the moment."))
-		return;
-	}
+	FString BodyString;
+	TSharedRef<TJsonWriter<TCHAR>> JsonWriter = TJsonWriterFactory<>::Create(&BodyString);
+	FJsonSerializer::Serialize(RequestJson.ToSharedRef(), JsonWriter);
+	HttpRequest->SetContentAsString(BodyString);
 
-	if (HttpResponse.IsValid())
+	HttpRequest->OnProcessRequestComplete().BindLambda([this, InCallback](FHttpRequestPtr HttpRequest, FHttpResponsePtr HttpResponse, bool bSucceeded)
 	{
-		if (EHttpResponseCodes::IsOk(HttpResponse->GetResponseCode()))
+		bool bSuccess = false;
+		if (!bSucceeded)
 		{
-			FString ResponseStr = HttpResponse->GetContentAsString();
-
-			TSharedPtr<FJsonObject> RootJsonObject;
-			TSharedRef<TJsonReader<>> Reader = TJsonReaderFactory<>::Create(ResponseStr);
-			FJsonSerializer::Deserialize(Reader, RootJsonObject);
-
-			if (RootJsonObject.IsValid())
+			UE_LOG(LogSubmitTool, Error, TEXT("Failed to refresh Jira OAuth token."));			
+		}
+		else
+		{
+			if (HttpResponse.IsValid())
 			{
+				if (EHttpResponseCodes::IsOk(HttpResponse->GetResponseCode()))
+				{
+					FJiraToken NewToken;
+					if (FJsonObjectConverter::JsonObjectStringToUStruct<FJiraToken>(*HttpResponse->GetContentAsString(), &NewToken))
+					{
+						ServiceProvider.Pin()->GetService<FCredentialsService>()->SetJiraToken(MoveTemp(NewToken));
+						UE_LOG(LogSubmitTool, Log, TEXT("Jira access token refreshed successfully."));
+						bSuccess = true;
+					}
+					else
+					{
+						UE_LOG(LogSubmitTool, Error, TEXT("Failed to parse refreshed Jira OAuth token."));
+						UE_LOG(LogSubmitToolDebug, Error, TEXT("%s"), *HttpResponse->GetContentAsString());
+					}
+				}
+				else
+				{
+					UE_LOG(LogSubmitTool, Error, TEXT("Jira token refresh failed with code %d."), HttpResponse->GetResponseCode());
+					UE_LOG(LogSubmitToolDebug, Error, TEXT("%s"), *HttpResponse->GetContentAsString());
+				}
 			}
 		}
-	}
+
+		InCallback(bSuccess);
+	});
+
+	HttpRequest->ProcessRequest();
+}
+
+const FString FJiraService::GetJiraIssuesFilepath() const
+{
+	return FPaths::Combine(FSubmitToolCoreUtils::GetLocalAppDataPath(), TEXT("SubmitTool"), TEXT("jira.issues.dat"));
 }
 
 void FJiraService::GetIssueAndCreateServiceDeskRequest(const FString& Key, const FString& Description, const FString& SwarmURL, const FString& InCurrentStream, const TMap<FString, TSharedPtr<FIntegrationOptionBase>>& InIntegrationOptions, const FOnBooleanValueChanged OnComplete)
@@ -746,6 +804,7 @@ void FJiraService::CreateServiceDeskRequest_HttpRequestComplete(FHttpRequestPtr
 			if(!FJsonSerializer::Deserialize(Reader, JsonObj))
 			{
 				UE_LOG(LogSubmitTool, Error, TEXT("Unable to deserialize swarm create response"));
+				OnComplete.ExecuteIfBound(false);
 				return;
 			}
 
@@ -770,4 +829,8 @@ void FJiraService::CreateServiceDeskRequest_HttpRequestComplete(FHttpRequestPtr
 			OnComplete.ExecuteIfBound(false);
 		}
 	}
+	else
+	{
+		OnComplete.ExecuteIfBound(false);
+	}
 }
diff --git a/Engine/Source/Programs/SubmitTool/Private/Logic/JiraService.h b/Engine/Source/Programs/SubmitTool/Private/Logic/JiraService.h
--- a/Engine/Source/Programs/SubmitTool/Private/Logic/JiraService.h
+++ b/Engine/Source/Programs/SubmitTool/Private/Logic/JiraService.h
@@ -48,15 +48,13 @@ public:
 	void GetIssueAndCreateServiceDeskRequest(const FString& Key, const FString& Description, const FString& SwarmURL, const FString& InCurrentStream, const TMap<FString, TSharedPtr<FIntegrationOptionBase>>& InIntegrationOptions, const FOnBooleanValueChanged OnComplete);
 	void CreateServiceDeskRequest(TSharedPtr<FJsonObject> InBaseJiraJsonObject, const FString& Description, const FString& SwarmURL, const FString& InCurrentStream, const TMap<FString, TSharedPtr<FIntegrationOptionBase>>& InIntegrationOptions, const FOnBooleanValueChanged OnComplete);
 private:
-	void QueryIssues(bool bInFailSilently = false);
-	void QueryIssues_HttpRequestComplete(FHttpRequestPtr HttpRequest, FHttpResponsePtr HttpResponse, bool bSucceeded, bool bFailSilently);
+	void QueryIssues(bool bInFailSilently = false, bool bAllowTokenRefresh = true);
+	void QueryIssues_HttpRequestComplete(FHttpRequestPtr HttpRequest, FHttpResponsePtr HttpResponse, bool bSucceeded, bool bFailSilently, bool bAllowTokenRefresh);
+	void RefreshOAuthToken(TFunction<void(bool)> InCallback);
 
 	bool ParseJsonObject(const TSharedPtr<FJsonObject>* InJsonObject, FJiraIssue& OutJiraIssue) const;
 
 	const FString GetJiraIssuesFilepath() const;
-	
-	void GetuserInfo();
-	void GetuserInfo_HttpRequestComplete(FHttpRequestPtr HttpRequest, FHttpResponsePtr HttpResponse, bool bSucceeded);
 
 	void SaveJiraIssues() const;
 	void LoadJiraIssues();
