Add STDErr pipe to process wrapper, correctly identifying stderr instead of treating everything as stdout

#Synced-CL 52011596

diff --git a/Engine/Source/Programs/SubmitTool/Private/Logic/CredentialsService.cpp b/Engine/Source/Programs/SubmitTool/Private/Logic/CredentialsService.cpp
--- a/Engine/Source/Programs/SubmitTool/Private/Logic/CredentialsService.cpp
+++ b/Engine/Source/Programs/SubmitTool/Private/Logic/CredentialsService.cpp
@@ -203,14 +203,14 @@ void FCredentialsService::GetOIDCToken()
 				TRACE_CPUPROFILER_EVENT_SCOPE(FCredentialsService::GetOIDCToken);
 
 				FString FullOutput;
-				FOnOutputLine OutputLineProcess = FOnOutputLine::CreateLambda([&FullOutput](const FString& OutputLine, const EProcessOutputType& OutputType) {
-					if(OutputType == EProcessOutputType::ProcessError)
+				FOnOutputLine OutputLineProcess = FOnOutputLine::CreateLambda([&FullOutput](FString OutputLine, const EProcessOutputType& OutputType) {
+					if(OutputType == EProcessOutputType::ProcessError || OutputType == EProcessOutputType::STDErr)
 					{
 						UE_LOG(LogSubmitTool, Error, TEXT("%s"), *OutputLine);
 					}
 					else
 					{
-						if (OutputType == EProcessOutputType::SDTOutput)
+						if (OutputType == EProcessOutputType::STDOutput)
 						{
 							FullOutput += OutputLine;
 						}
diff --git a/Engine/Source/Programs/SubmitTool/Private/Logic/ProcessPipes.cpp b/Engine/Source/Programs/SubmitTool/Private/Logic/ProcessPipes.cpp
--- a/Engine/Source/Programs/SubmitTool/Private/Logic/ProcessPipes.cpp
+++ b/Engine/Source/Programs/SubmitTool/Private/Logic/ProcessPipes.cpp
@@ -24,6 +24,20 @@ bool FProcessPipes::Create()
 		return false;
 	}
 
+	if (!FPlatformProcess::CreatePipe(StdErrReadPipe, StdErrWritePipe, false))
+	{
+		// Cleanup StdOut & StdIn pipes before returning
+		FPlatformProcess::ClosePipe(StdOutReadPipe, StdOutWritePipe);
+		FPlatformProcess::ClosePipe(StdInReadPipe, StdInWritePipe);
+
+		StdInReadPipe = nullptr;
+		StdInWritePipe = nullptr;
+		StdOutReadPipe = nullptr;
+		StdOutWritePipe = nullptr;
+
+		return false;
+	}
+
 	return true;
 }
 
@@ -40,6 +54,12 @@ void FProcessPipes::Reset()
 
 		StdInReadPipe = nullptr;
 		StdInWritePipe = nullptr;
+
+
+		FPlatformProcess::ClosePipe(StdErrReadPipe, StdErrWritePipe);
+
+		StdErrReadPipe = nullptr;
+		StdErrWritePipe = nullptr;
 	}
 }
 
@@ -50,5 +70,7 @@ bool FProcessPipes::IsValid() const
 	return	StdInReadPipe != nullptr &&
 			StdInWritePipe != nullptr &&
 			StdOutWritePipe != nullptr &&
-			StdOutReadPipe != nullptr;
+			StdOutReadPipe != nullptr&&
+			StdErrWritePipe != nullptr &&
+			StdErrReadPipe != nullptr;
 }
\ No newline at end of file
diff --git a/Engine/Source/Programs/SubmitTool/Private/Logic/ProcessPipes.h b/Engine/Source/Programs/SubmitTool/Private/Logic/ProcessPipes.h
--- a/Engine/Source/Programs/SubmitTool/Private/Logic/ProcessPipes.h
+++ b/Engine/Source/Programs/SubmitTool/Private/Logic/ProcessPipes.h
@@ -2,7 +2,7 @@
 
 #pragma once
 
-/** Wrapper around stdin and stdout pipes created by a call to FPlatformProcess::CreatePipe */
+/** Wrapper around stdin, stderr and stdout pipes created by a call to FPlatformProcess::CreatePipe */
 struct FProcessPipes
 {
 	FProcessPipes() = default;
@@ -38,11 +38,25 @@ struct FProcessPipes
 		return StdOutReadPipe;
 	}
 
+	void* GetStdErrForProcess() const
+	{
+		return StdErrWritePipe;
+	}
+
+	void* GetStdErrForReading() const
+	{
+		return StdErrReadPipe;
+	}
+
 private:
 
 	void* StdOutReadPipe = nullptr;
 	void* StdOutWritePipe = nullptr;
 
+	void* StdErrReadPipe = nullptr;
+	void* StdErrWritePipe = nullptr;
+
 	void* StdInReadPipe = nullptr;
 	void* StdInWritePipe = nullptr;
+
 };
diff --git a/Engine/Source/Programs/SubmitTool/Private/Logic/ProcessWrapper.cpp b/Engine/Source/Programs/SubmitTool/Private/Logic/ProcessWrapper.cpp
--- a/Engine/Source/Programs/SubmitTool/Private/Logic/ProcessWrapper.cpp
+++ b/Engine/Source/Programs/SubmitTool/Private/Logic/ProcessWrapper.cpp
@@ -42,7 +42,8 @@ bool FProcessWrapper::Start(bool bWaitForExit)
 		return false;
 	}
 	
-	OutputRemainder = FString();
+	OutputRemainder.Reset();
+	StdErrOutputRemainder.Reset();
 	ExecutingTime = 0;
 
 	if (!Pipes->Create())
@@ -67,7 +68,8 @@ bool FProcessWrapper::Start(bool bWaitForExit)
 		0,
 		WorkingDir.Len() ? *WorkingDir : nullptr,
 		Pipes->GetStdOutForProcess(),
-		Pipes->GetStdInForProcess())
+		Pipes->GetStdInForProcess(),
+		Pipes->GetStdErrForProcess())
 	);
 
 	if (!ProcessHandle->IsValid())
@@ -119,30 +121,29 @@ bool FProcessWrapper::OnTick(float Delta)
 		return false;
 	}
 
-	if(FPlatformProcess::IsProcRunning(*ProcessHandle))
-	{
-		ReadOutput(false);
-		ExecutingTime += Delta;
-		return true;
-	}
-
+	bool bIsProcessRunning = FPlatformProcess::IsProcRunning(*ProcessHandle);
 	ExecutingTime += Delta;
 
-	// Flush output
-	ReadOutput(true);	
-	FPlatformProcess::GetProcReturnCode(*ProcessHandle, &ExitCode);
-	bIsComplete = true;
+	// Flush output	only if process is finished
+	ReadOutput(!bIsProcessRunning);
 
-	OutputLine(FString::Printf(TEXT("Completed running process %s. Process took %s and exited with code %d"), *ProcessName, *FGenericPlatformTime::PrettyTime(ExecutingTime), ExitCode), EProcessOutputType::ProcessInfo);
+	if(!bIsProcessRunning)
+	{
+		FPlatformProcess::GetProcReturnCode(*ProcessHandle, &ExitCode);
+		bIsComplete = true;
 
-	Cleanup();
+		OutputLine(FString::Printf(TEXT("Completed running process %s. Process took %s and exited with code %d"), *ProcessName, *FGenericPlatformTime::PrettyTime(ExecutingTime), ExitCode), EProcessOutputType::ProcessInfo);
 
-	if(OnCompleted.IsBound())
-	{
-		OnCompleted.ExecuteIfBound(ExitCode);
+		Cleanup();
+
+		if(OnCompleted.IsBound())
+		{
+			OnCompleted.ExecuteIfBound(ExitCode);
+		}
 	}
 
-	return false;
+	// Requeue Tick
+	return bIsProcessRunning;
 }
 
 void FProcessWrapper::Cleanup()
@@ -160,35 +161,61 @@ void FProcessWrapper::ReadOutput(bool FlushOutput)
 	{
 		return;
 	}
-	
-	FString NewOutput = OutputRemainder + FPlatformProcess::ReadPipe(Pipes->GetStdOutForReading());
 
-	if(!NewOutput.IsEmpty() && (NewOutput.Contains(TEXT("\n")) || FlushOutput))
+	//STDOut
+	ProcessLines(FlushOutput, false);
+
+	//STDErr
+	ProcessLines(FlushOutput, true);
+}
+
+void FProcessWrapper::ProcessLines(bool bFlushOutput, bool bIsStdErr)
+{
+	FString OutputLines;
+	if (bIsStdErr)
 	{
-		// Find the last line, which may be truncated
-		if(!FlushOutput)
+		OutputLines = StdErrOutputRemainder + FPlatformProcess::ReadPipe(Pipes->GetStdErrForReading());
+	}
+	else
+	{
+		OutputLines = OutputRemainder + FPlatformProcess::ReadPipe(Pipes->GetStdOutForReading());
+	}
+
+	if (!bFlushOutput)
+	{
+		FString& OutputRemainderRef = bIsStdErr ? StdErrOutputRemainder : OutputRemainder;
+		if (!OutputLines.Contains(TEXT("\n")))
 		{
-			int32 Position;
-			NewOutput.FindLastChar('\n', Position);
-			OutputRemainder = NewOutput.Mid(Position + 1);
-			NewOutput.RemoveFromEnd(OutputRemainder);
+			// If we're not flushing and we only have a partial line, exit here and keep everything as remainder
+			OutputRemainderRef = OutputLines;
+			return;
 		}
 
-		TArray<FString> Lines;
-		const TCHAR* Separators[] = { TEXT("\n"), TEXT("\r") };
-		NewOutput.ParseIntoArray(Lines, Separators, UE_ARRAY_COUNT(Separators));
+		// Find the last line, which may be truncated and store it in the output remainder for the next processline
+		int32 Position;
+		OutputLines.FindLastChar('\n', Position);
+		OutputRemainderRef = OutputLines.Mid(Position + 1);
+		OutputLines.RemoveFromEnd(OutputRemainderRef);
+	}
+	else
+	{
+		// if we've flushed, empty previous output remainder
+		FString& OutputRemainderRef = bIsStdErr ? StdErrOutputRemainder : OutputRemainder;
+		OutputRemainderRef.Empty();
+	}
 
-		for(const FString& Line : Lines)
-		{
-			OutputLine(Line, EProcessOutputType::SDTOutput);
-		}
+	// Parse into lines and fire line events
+	TArray<FString> Lines;
+	const TCHAR* Separators[] = { TEXT("\n"), TEXT("\r") };
+	OutputLines.ParseIntoArray(Lines, Separators, UE_ARRAY_COUNT(Separators));
+
+	for (FString& Line : Lines)
+	{
+		OutputLine(MoveTemp(Line), bIsStdErr ? EProcessOutputType::STDErr : EProcessOutputType::STDOutput);
 	}
 }
 
-void FProcessWrapper::OutputLine(const FString OutputLine, const EProcessOutputType& OutputType)
+void FProcessWrapper::OutputLine(FString&& OutputLine, const EProcessOutputType& OutputType)
 {
-	if(OnOutputLine.IsBound())
-	{
-		OnOutputLine.ExecuteIfBound(OutputLine, OutputType);
-	}
+	OnOutputLine.ExecuteIfBound(MoveTemp(OutputLine), OutputType);
 }
diff --git a/Engine/Source/Programs/SubmitTool/Private/Logic/ProcessWrapper.h b/Engine/Source/Programs/SubmitTool/Private/Logic/ProcessWrapper.h
--- a/Engine/Source/Programs/SubmitTool/Private/Logic/ProcessWrapper.h
+++ b/Engine/Source/Programs/SubmitTool/Private/Logic/ProcessWrapper.h
@@ -6,12 +6,13 @@
 
 enum EProcessOutputType
 {
-	SDTOutput,
+	STDOutput,
+	STDErr,
 	ProcessInfo,
 	ProcessError,
 };
 
-DECLARE_DELEGATE_TwoParams(FOnOutputLine, const FString& /*Output*/, const EProcessOutputType&)
+DECLARE_DELEGATE_TwoParams(FOnOutputLine, FString /*Output*/, const EProcessOutputType&)
 DECLARE_DELEGATE_OneParam(FOnCompleted, const int32 /*Return code*/)
 
 class FProcessWrapper
@@ -35,7 +36,8 @@ private:
 	bool OnTick(float Delta);
 	void Cleanup();
 	void ReadOutput(bool FlushOutput);
-	void OutputLine(const FString Output, const EProcessOutputType& OutputType);
+	void ProcessLines(bool bFlushOutput, bool bIsStdErr);
+	void OutputLine(FString&& Output, const EProcessOutputType& OutputType);
 	
 	FString ProcessName;
 	FString Path;
@@ -49,6 +51,7 @@ private:
 	TUniquePtr<FProcHandle> ProcessHandle;
 	TUniquePtr<FProcessPipes> Pipes;
 	FString OutputRemainder;
+	FString StdErrOutputRemainder;
 
 	FTSTicker::FDelegateHandle TickerHandle;
 
diff --git a/Engine/Source/Programs/SubmitTool/Private/Logic/UpdateService.cpp b/Engine/Source/Programs/SubmitTool/Private/Logic/UpdateService.cpp
--- a/Engine/Source/Programs/SubmitTool/Private/Logic/UpdateService.cpp
+++ b/Engine/Source/Programs/SubmitTool/Private/Logic/UpdateService.cpp
@@ -6,6 +6,7 @@
 #include "HttpModule.h"
 #include "Interfaces/IHttpResponse.h"
 #include "JsonObjectConverter.h"
+#include "Misc/CommandLine.h"
 #include "Misc/FileHelper.h"
 #include "Misc/Paths.h"
 
@@ -104,7 +105,7 @@ void FUpdateService::StartAutoUpdateScript()
 			Cmd,
 			Args,
 			nullptr,
-			FOnOutputLine::CreateLambda([](const FString& InOutput, const EProcessOutputType& OutputType){ UE_LOG(LogSubmitToolDebug, Log, TEXT("AutoUpdateOutput: %s"), *InOutput); }),
+			FOnOutputLine::CreateLambda([](FString InOutput, const EProcessOutputType& OutputType){ UE_LOG(LogSubmitToolDebug, Log, TEXT("AutoUpdateOutput: %s"), *InOutput); }),
 			FString(),
 			bLaunchesHidden,
 			bLaunchesReallyHidden,
diff --git a/Engine/Source/Programs/SubmitTool/Private/Logic/Validators/ValidatorRunExecutable.cpp b/Engine/Source/Programs/SubmitTool/Private/Logic/Validators/ValidatorRunExecutable.cpp
--- a/Engine/Source/Programs/SubmitTool/Private/Logic/Validators/ValidatorRunExecutable.cpp
+++ b/Engine/Source/Programs/SubmitTool/Private/Logic/Validators/ValidatorRunExecutable.cpp
@@ -228,9 +228,9 @@ bool FValidatorRunExecutable::QueueProcess(const FString& InId, const FString& L
 	ProcessData.bIgnoringWarningsAsErrors = !GetTypedDefinition<FValidatorRunExecutableDefinition>()->EnableWarningsAsErrorsAnchor.IsEmpty();
 	ProcessesData.FindOrAdd(InId) = ProcessData;
 
-	FOnOutputLine LineHandler = FOnOutputLine::CreateLambda([this, InId](const FString& InOutput, const EProcessOutputType& InType)
+	FOnOutputLine LineHandler = FOnOutputLine::CreateLambda([this, InId](FString InOutput, const EProcessOutputType& InType)
 	{
-		OnProcessOutputLine(InId, InOutput, InType);
+		OnProcessOutputLine(InId, MoveTemp(InOutput), InType);
 	});
 
 	FOnCompleted CompleteHandler = FOnCompleted::CreateLambda([this, InId](const int& InResult)
@@ -241,7 +241,7 @@ bool FValidatorRunExecutable::QueueProcess(const FString& InId, const FString& L
 	return true;
 }
 
-void FValidatorRunExecutable::OnProcessOutputLine(const FString& ProcessId, const FString& Line, const EProcessOutputType& OutputType)
+void FValidatorRunExecutable::OnProcessOutputLine(const FString& ProcessId, FString&& Line, const EProcessOutputType& OutputType)
 {
 
 	const FString& FormattedMessage = FString::Printf(TEXT("[%s] [%s] %s"), *GetValidatorName(), *ProcessId, *Line);
@@ -263,7 +263,7 @@ void FValidatorRunExecutable::OnProcessOutputLine(const FString& ProcessId, cons
 	}
 
 	bool bIsWarning = false;
-	bool bIsError = IsLineAnError(Line, bIsWarning, ProcessesData[ProcessId].bIgnoringOutputError, ProcessesData[ProcessId].bIgnoringWarningsAsErrors);
+	bool bIsError = IsLineAnError(Line, OutputType, bIsWarning, ProcessesData[ProcessId].bIgnoringOutputError, ProcessesData[ProcessId].bIgnoringWarningsAsErrors);
 
 	if(bIsError)
 	{
@@ -295,7 +295,7 @@ void FValidatorRunExecutable::OnProcessOutputLine(const FString& ProcessId, cons
 	}
 }
 
-bool FValidatorRunExecutable::IsLineAnError(const FString& InLine, bool& OutbIsWarning, bool& InOutbIgnoringOutputErrors, bool& InOutbIgnoringWarningsAsErrors)
+bool FValidatorRunExecutable::IsLineAnError(const FString& InLine, const EProcessOutputType InOutputType, bool& OutbIsWarning, bool& InOutbIgnoringOutputErrors, bool& InOutbIgnoringWarningsAsErrors)
 {
 	const FValidatorRunExecutableDefinition* TypedDefinition = GetTypedDefinition<FValidatorRunExecutableDefinition>();
 
@@ -345,6 +345,11 @@ bool FValidatorRunExecutable::IsLineAnError(const FString& InLine, bool& OutbIsW
 		}
 	}
 
+	if (InOutputType == EProcessOutputType::STDErr)
+	{
+		return true;
+	}
+
 	{
 		FRegexMatcher Regex = FRegexMatcher(*TypedDefinition->BuiltRegexError, InLine);
 		if(Regex.FindNext())
diff --git a/Engine/Source/Programs/SubmitTool/Private/Logic/Validators/ValidatorRunExecutable.h b/Engine/Source/Programs/SubmitTool/Private/Logic/Validators/ValidatorRunExecutable.h
--- a/Engine/Source/Programs/SubmitTool/Private/Logic/Validators/ValidatorRunExecutable.h
+++ b/Engine/Source/Programs/SubmitTool/Private/Logic/Validators/ValidatorRunExecutable.h
@@ -41,8 +41,8 @@ public:
 	 
 protected:
 	virtual bool QueueProcess(const FString& ProcessId, const FString& LocalPath, const FString& Args);
-	virtual void OnProcessOutputLine(const FString& ProcessId, const FString& Line, const EProcessOutputType& OutputType);
-	virtual bool IsLineAnError(const FString& InLine, bool& OutbIsWarning, bool& InOutbIgnoringOutputErrors, bool& InOutbIgnoringWarningsAsErrors);
+	virtual void OnProcessOutputLine(const FString& ProcessId, FString&& Line, const EProcessOutputType& OutputType);
+	virtual bool IsLineAnError(const FString& InLine, const EProcessOutputType OutputType, bool& OutbIsWarning, bool& InOutbIgnoringOutputErrors, bool& InOutbIgnoringWarningsAsErrors);
 	virtual void OnProcessComplete(const FString& ProcessId, int32 ReturnCode);
 	
 	virtual void Tick(float DeltaTime) override;
diff --git a/Engine/Source/Programs/SubmitTool/Private/Models/ModelInterface.cpp b/Engine/Source/Programs/SubmitTool/Private/Models/ModelInterface.cpp
--- a/Engine/Source/Programs/SubmitTool/Private/Models/ModelInterface.cpp
+++ b/Engine/Source/Programs/SubmitTool/Private/Models/ModelInterface.cpp
@@ -264,9 +264,9 @@ void FModelInterface::GenerateCLDescription()
 	{
 		GeneratedDescription.Empty();
 
-		FOnOutputLine LineHandler = FOnOutputLine::CreateLambda([this](const FString& InOutput, const EProcessOutputType& InType)
+		FOnOutputLine LineHandler = FOnOutputLine::CreateLambda([this](FString InOutput, const EProcessOutputType& InType)
 		{
-			if (InType == EProcessOutputType::ProcessError)
+			if (InType == EProcessOutputType::ProcessError || InType == EProcessOutputType::STDErr)
 			{
 				UE_LOG(LogSubmitTool, Error, TEXT("[Description Generation]: %s"), *InOutput);
 			}
