Fix telemetry submitted changelist number not being captured when a CL submission succeeds

#Synced-CL 52011596
#Synced-CL 52011596

diff --git a/Engine/Source/Programs/SubmitTool/Private/Models/ModelInterface.cpp b/Engine/Source/Programs/SubmitTool/Private/Models/ModelInterface.cpp
--- a/Engine/Source/Programs/SubmitTool/Private/Models/ModelInterface.cpp
+++ b/Engine/Source/Programs/SubmitTool/Private/Models/ModelInterface.cpp
@@ -1090,6 +1090,10 @@ void FModelInterface::OnSubmitOperationComplete(const FSCCommand& InCmd)
 			CLNumber = 0;
 			SubmittedCL = TEXT("0");
 		}
+		else
+		{
+			SubmittedCL = *SubmittedCLString;
+		}
 
 
 		TArray<FAnalyticsEventAttribute> SubmitSuccessAttribs = MakeAnalyticsEventAttributeArray(
