Add a callback for when the full validation queue finishes to automatically scroll the summary log to the end on failure
Fix incorrect handle being passed to RemoveValidationFinishedCallback
Switch validator correlation ID logging to use the debug log category

#Synced-CL 52011596

diff --git a/Engine/Source/Programs/SubmitTool/Private/Logic/Validators/ValidatorBase.cpp b/Engine/Source/Programs/SubmitTool/Private/Logic/Validators/ValidatorBase.cpp
--- a/Engine/Source/Programs/SubmitTool/Private/Logic/Validators/ValidatorBase.cpp
+++ b/Engine/Source/Programs/SubmitTool/Private/Logic/Validators/ValidatorBase.cpp
@@ -187,7 +187,7 @@ void FValidatorBase::StartValidation()
 				UE_LOG(LogValidators, Log, TEXT("[%s] Skipping Files because they were already validated in a previous execution:\n%s"), *GetValidatorName(), *FileList);
 			}
 
-			UE_LOG(LogValidators, Log, TEXT("[%s] CorrelationId: %s"), *GetValidatorName(), *CorrelationId.ToString());
+			UE_LOG(LogSubmitToolDebug, Log, TEXT("[%s] CorrelationId: %s"), *GetValidatorName(), *CorrelationId.ToString());
 			if (!Validate(ChangelistService->GetCLDescription(), FilteredFiles, TagService->GetTagsArray()))
 			{
 				if(FailureReason == EFailureReason::None)
diff --git a/Engine/Source/Programs/SubmitTool/Private/View/Widgets/SubmitToolWidget.cpp b/Engine/Source/Programs/SubmitTool/Private/View/Widgets/SubmitToolWidget.cpp
--- a/Engine/Source/Programs/SubmitTool/Private/View/Widgets/SubmitToolWidget.cpp
+++ b/Engine/Source/Programs/SubmitTool/Private/View/Widgets/SubmitToolWidget.cpp
@@ -86,6 +86,7 @@ void SubmitToolWidget::Construct(const FArguments& InArgs, TSharedPtr<FSubmitToo
 
 	FSlateApplication::Get().OnApplicationActivationStateChanged().AddRaw(this, &SubmitToolWidget::HandleApplicationActivationStateChanged);
 	OnValidatorFinishedHandle = ModelInterface->AddSingleValidatorFinishedCallback(FOnSingleTaskFinished::FDelegate::CreateRaw(this, &SubmitToolWidget::OnSingleValidatorFinished));
+	OnValidationQueueFinishedHandle = ModelInterface->AddValidationFinishedCallback(FOnTaskFinished::FDelegate::CreateRaw(this, &SubmitToolWidget::OnValidationQueueFinished));
 	OnValidationUpdateHandle = ModelInterface->AddValidationUpdatedCallback(FOnTaskRunStateChanged::FDelegate::CreateRaw(this, &SubmitToolWidget::OnValidationUpdated));
 	OnCLDescriptionUpdatedHandle = ModelInterface->GetCLDescriptionUpdatedDelegate().Add(FOnCLDescriptionUpdated::FDelegate::CreateRaw(this, &SubmitToolWidget::OnCLDescriptionUpdated));
 
@@ -444,7 +445,8 @@ SubmitToolWidget::~SubmitToolWidget()
 		UserPrefs->ValidatorSectionSize = ValidatorSectionSlot->GetSizeValue();
 		UserPrefs->LogSectionSize = LogSectionSlot->GetSizeValue();
 
-		ModelInterface->RemoveValidationFinishedCallback(OnValidatorFinishedHandle);
+		ModelInterface->RemoveValidationFinishedCallback(OnValidationQueueFinishedHandle);
+		ModelInterface->RemoveSingleValidatorFinishedCallback(OnValidatorFinishedHandle);
 		ModelInterface->RemoveValidationUpdatedCallback(OnValidationUpdateHandle);
 		ModelInterface->GetCLDescriptionUpdatedDelegate().Remove(OnCLDescriptionUpdatedHandle);
 		ModelInterface = nullptr;
@@ -1025,6 +1027,14 @@ void SubmitToolWidget::OnSingleValidatorFinished(const FValidatorBase& InValidat
 	}
 }
 
+void SubmitToolWidget::OnValidationQueueFinished(const bool bInSuccess)
+{
+	if (!bInSuccess)
+	{
+		SummaryLogWidget->ScrollToEnd();
+	}
+}
+
 void SubmitToolWidget::OnValidationUpdated(bool bValid)
 {
 	if(bValid)
diff --git a/Engine/Source/Programs/SubmitTool/Private/View/Widgets/SubmitToolWidget.h b/Engine/Source/Programs/SubmitTool/Private/View/Widgets/SubmitToolWidget.h
--- a/Engine/Source/Programs/SubmitTool/Private/View/Widgets/SubmitToolWidget.h
+++ b/Engine/Source/Programs/SubmitTool/Private/View/Widgets/SubmitToolWidget.h
@@ -67,9 +67,11 @@ private:
 	SSplitter::FSlot* LogSectionSlot;
 
 	FDelegateHandle OnValidatorFinishedHandle;
+	FDelegateHandle OnValidationQueueFinishedHandle;
 	FDelegateHandle OnValidationUpdateHandle;
 	FDelegateHandle OnCLDescriptionUpdatedHandle;
 	void OnSingleValidatorFinished(const FValidatorBase& InValidator);
+	void OnValidationQueueFinished(const bool bInSuccess);
 	void OnValidationUpdated(bool bValid);
 
 #if PLATFORM_WINDOWS
