Added bCanBeUserStopped property
Modified how Required status is shown in the UI to distinguish between required to succeed and required to complete
Modifying internal config for validators

#Synced-CL 52444213

diff --git a/Engine/Source/Developer/SubmitToolCore/Public/Logic/Validators/ValidatorDefinition.h b/Engine/Source/Developer/SubmitToolCore/Public/Logic/Validators/ValidatorDefinition.h
--- a/Engine/Source/Developer/SubmitToolCore/Public/Logic/Validators/ValidatorDefinition.h
+++ b/Engine/Source/Developer/SubmitToolCore/Public/Logic/Validators/ValidatorDefinition.h
@@ -68,6 +68,12 @@ struct FValidatorDefinition
 	UPROPERTY()
 	bool bCanBeDisabledByUser = true;
 
+	///
+	///	If false, this validator won't have the Stop button available
+	/// 
+	UPROPERTY()
+	bool bCanBeStoppedByUser = true;
+
 	/// 
 	/// Whether the validator is required to allow submission or not, a failing required validation will always block submission
 	/// 
diff --git a/Engine/Source/Programs/SubmitTool/Private/Version/AppVersionDefines.h b/Engine/Source/Programs/SubmitTool/Private/Version/AppVersionDefines.h
--- a/Engine/Source/Programs/SubmitTool/Private/Version/AppVersionDefines.h
+++ b/Engine/Source/Programs/SubmitTool/Private/Version/AppVersionDefines.h
@@ -21,7 +21,7 @@
 
 #define SUBMIT_TOOL_VERSION_MAJOR 6
 #define SUBMIT_TOOL_VERSION_MINOR 7
-#define SUBMIT_TOOL_VERSION_PATCH 2
+#define SUBMIT_TOOL_VERSION_PATCH 3
 
 #ifndef SUBMIT_TOOL_CHANGELIST
 #define SUBMIT_TOOL_CHANGELIST_STRING TEXT("NoCL")
diff --git a/Engine/Source/Programs/SubmitTool/Private/View/Widgets/ValidatorsWidget.cpp b/Engine/Source/Programs/SubmitTool/Private/View/Widgets/ValidatorsWidget.cpp
--- a/Engine/Source/Programs/SubmitTool/Private/View/Widgets/ValidatorsWidget.cpp
+++ b/Engine/Source/Programs/SubmitTool/Private/View/Widgets/ValidatorsWidget.cpp
@@ -274,7 +274,7 @@ public:
 
 		if(InColumnName == ValidatorsColumns::Name)
 		{
-			const FString NameStr = Node.Pin()->Definition->IsRequired ? Node.Pin()->GetValidatorName() : Node.Pin()->GetValidatorName() + TEXT(" - Optional");
+			const FString NameStr = Node.Pin()->Definition->IsRequired || Node.Pin()->Definition->bRequireCompleteWhenOptional ? Node.Pin()->GetValidatorName() : Node.Pin()->GetValidatorName() + TEXT(" - Optional");
 			TSharedRef<STextBlock> Name = SNew(STextBlock)
 				.Text(FText::FromString(NameStr))
 				.ToolTipText(FText::FromString(Node.Pin()->Definition->ToolTip));
@@ -353,9 +353,23 @@ public:
 		}
 		else if (InColumnName == ValidatorsColumns::Required)
 		{
-			TSharedRef<STextBlock> Required = SNew(STextBlock).Text(FText::FromString(Node.Pin()->Definition->IsRequired ? TEXT("Required") : TEXT("Optional")));
+			FText Text;			
+			if (Node.Pin()->Definition->IsRequired)
+			{
+				Text = FText::FromString(TEXT("Success"));
+			}
+			else if(Node.Pin()->Definition->bRequireCompleteWhenOptional)
+			{
+				Text = FText::FromString(TEXT("Complete"));
+			}
+			else
+			{
+				Text = FText::FromString(TEXT("Optional"));
+			}
+
 			box->SetContent(
-				Required
+				SNew(STextBlock)
+				.Text(Text)
 			);
 		}
 		else if(InColumnName == ValidatorsColumns::Status)
@@ -403,10 +417,20 @@ public:
 				.Padding(3,0,0,0)
 				[
 					SNew(SButton)
-					.IsEnabled_Lambda([this]() {return ModelInterface->GetInputEnabled() && Node.Pin()->GetValidatorState() != EValidationStates::Disabled && bIsRelevantToCL; })
+					.IsEnabled_Lambda([this]() {return (!Node.Pin()->GetIsRunningOrQueued() || Node.Pin()->Definition->bCanBeStoppedByUser) && ModelInterface->GetInputEnabled() && Node.Pin()->GetValidatorState() != EValidationStates::Disabled && bIsRelevantToCL; })
 					.HAlign(EHorizontalAlignment::HAlign_Center)
 					.VAlign(EVerticalAlignment::VAlign_Center)
-					.ToolTipText(FText::FromString(TEXT("Runs this validator. To force run, hold CTRL.")))
+					.ToolTipText_Lambda([this]
+					{
+						if (Node.Pin()->GetIsRunningOrQueued())
+						{
+							return FText::FromString(Node.Pin()->Definition->bCanBeStoppedByUser ? TEXT("Stops this validator.") : TEXT("Stopping this validator is disabled"));
+						}
+						else
+						{
+							return FText::FromString(TEXT("Runs this validator. To force run, hold CTRL."));
+						}
+					})						
 					.OnClicked_Lambda([this]()
 					{
 						if (Node.Pin()->GetIsRunningOrQueued())
@@ -447,8 +471,7 @@ public:
 				.Padding(3, 0, 3, 0)
 				[
 					SNew(SButton)
-					.IsEnabled(&FModelInterface::GetInputEnabled)
-					.Visibility_Lambda([this]{ return Node.Pin()->Definition->bCanBeDisabledByUser ? EVisibility::All : EVisibility::Collapsed; })
+					.IsEnabled_Lambda([this] { return Node.Pin()->Definition->bCanBeDisabledByUser && ModelInterface->GetInputEnabled(); })
 					.HAlign(EHorizontalAlignment::HAlign_Center)
 					.VAlign(EVerticalAlignment::VAlign_Center)
 					.OnClicked_Lambda([this]() 
