Move service contracts (IChangelistService, ITagService, ILockdownService, ICredentialsService) and shared types (FJiraToken, LockdownData) into SubmitToolCore
update ModelInterface members and all GetService callsites to use interface types
Add SubmitToolRestricted conditional module integration and a preflight suggestion override extension point in SubmitToolCore

#Synced-CL 52444213

diff --git a/Engine/Source/Developer/SubmitToolCore/Public/Logic/Services/Interfaces/ICacheDataService.h b/Engine/Source/Developer/SubmitToolCore/Public/Logic/Services/Interfaces/ICacheDataService.h
new file mode 100644
--- /dev/null
+++ b/Engine/Source/Developer/SubmitToolCore/Public/Logic/Services/Interfaces/ICacheDataService.h
@@ -0,0 +1,29 @@
+// Copyright Epic Games, Inc. All Rights Reserved.
+
+#pragma once
+
+#include "CoreMinimal.h"
+#include "Logic/Services/Interfaces/ISubmitToolService.h"
+#include "Models/CacheData.h"
+#include "Models/SCFile.h"
+
+struct FGeneralParameters;
+
+class ICacheDataService : public ISubmitToolService
+{
+public:
+	virtual bool GetChangelistCacheData(const FString& InCLId, FChangelistCacheData& OutCLData) const = 0;
+
+	virtual const FDateTime GetLastValidationDate(const FString& InCLId, const FName& InValidatorId, const FString& InValidatorConfig, const FString& InDepotPath) const = 0;
+	virtual void UpdateLastValidationForFiles(const FString& InCLId, const FName& InValidatorId, const FString& InValidatorConfig, const TArray<FSCFileRef>& InFiles, const FDateTime& InNewTimestamp) = 0;
+
+	virtual FString GetIntegrationFieldValue(const FString& InCLId, const FString& InIntegrationFieldId) const = 0;
+	virtual void SetIntegrationFieldValue(const FString& InCLId, const FString& InIntegrationFieldId, const FString& InValue) = 0;
+
+	virtual FDateTime GetLastKnownShelfUpToDate(const FString& InCLId) const = 0;
+	virtual void SetLastShelfKnownUpToDate(const FString& InCLId, const FDateTime& InDateTime) = 0;
+
+	virtual void SaveCacheToDisk() = 0;
+};
+
+Expose_TNameOf(ICacheDataService);
diff --git a/Engine/Source/Developer/SubmitToolCore/Public/Logic/Services/Interfaces/IChangelistService.h b/Engine/Source/Developer/SubmitToolCore/Public/Logic/Services/Interfaces/IChangelistService.h
new file mode 100644
--- /dev/null
+++ b/Engine/Source/Developer/SubmitToolCore/Public/Logic/Services/Interfaces/IChangelistService.h
@@ -0,0 +1,46 @@
+// Copyright Epic Games, Inc. All Rights Reserved.
+
+#pragma once
+
+#include "Logic/Services/Interfaces/ISubmitToolService.h"
+#include "Logic/Services/Interfaces/ISTSourceControlService.h"
+#include "Logic/Validators/ValidatorDefinition.h"
+#include "Models/SCFile.h"
+
+DECLARE_MULTICAST_DELEGATE(FOnCLDescriptionUpdated)
+DECLARE_DELEGATE_OneParam(FOnChangeListReadyDelegate, bool /*bValidState*/)
+DECLARE_DELEGATE_OneParam(FOnChangelistRefreshDelegate, ETaskArea /*ChangedValue*/)
+
+class IChangelistService : public ISubmitToolService
+{
+public:
+	virtual const FString GetCLID() const = 0;
+	virtual const FString& GetCLDescription() = 0;
+	virtual bool SetCLDescription(const FString& NewDescription, bool bNotifyEvent = false) = 0;
+	virtual const TArray<FSCFileRef>& GetFilesInCL() const = 0;
+	virtual const TArray<FSCFileRef>& GetShelvedFilesInCL() const = 0;
+	virtual bool HasShelvedFiles() const = 0;
+	virtual bool HasP4OperationsRunning() const = 0;
+	virtual bool IsShelfReady() const = 0;
+	virtual bool IsShelfComplete() const = 0;
+	virtual const TMap<FString, FString>& GetOpenFilesInChangelists() const = 0;
+
+	virtual void Init() = 0;
+	virtual void Submit(const FString& InDescriptionAddendum = TEXT(""), TFunction<void(const FSCCommand&)> OnComplete = nullptr) = 0;
+	virtual void CreateCLFromDefaultCL() = 0;
+	virtual void FetchChangelistDataAsync() = 0;
+	virtual void RevertUnchangedFiles(TFunction<void(const FSCCommand&)> OnComplete = nullptr) = 0;
+	virtual void UpdateP4CLDescription() = 0;
+	virtual void UpdateP4CLDescriptionSynchronously() = 0;
+	virtual void DeleteShelvedFiles(TFunction<void(const FSCCommand&)> OnComplete = nullptr) = 0;
+	virtual void ReplaceShelvedFiles(TFunction<void(const FSCCommand&)> OnComplete = nullptr) = 0;
+	virtual void CreateShelvedFiles(TFunction<void(const FSCCommand&)> OnComplete = nullptr) = 0;
+	virtual void UpdateShelvedFiles(TArray<FString> InFiles, TFunction<void(const FSCCommand&)> OnComplete = nullptr) = 0;
+	virtual void CheckShelfIsUpToDate(TFunction<void(bool, TArray<FString>)> OnComplete) = 0;
+	virtual void EnsureShelfIsCurrent() = 0;
+	virtual void CancelP4Operations() = 0;
+
+	FOnCLDescriptionUpdated OnCLDescriptionUpdated;
+};
+
+Expose_TNameOf(IChangelistService);
diff --git a/Engine/Source/Developer/SubmitToolCore/Public/Logic/Services/Interfaces/ICredentialsService.h b/Engine/Source/Developer/SubmitToolCore/Public/Logic/Services/Interfaces/ICredentialsService.h
new file mode 100644
--- /dev/null
+++ b/Engine/Source/Developer/SubmitToolCore/Public/Logic/Services/Interfaces/ICredentialsService.h
@@ -0,0 +1,27 @@
+// Copyright Epic Games, Inc. All Rights Reserved.
+
+#pragma once
+
+#include "CoreMinimal.h"
+#include "Logic/Services/Interfaces/ISubmitToolService.h"
+#include "Models/JiraToken.h"
+#include "Tasks/Task.h"
+
+class ICredentialsService : public ISubmitToolService
+{
+public:
+	virtual bool HasJiraCredentials() const = 0;
+	virtual const FString& GetEncodedLoginString() const = 0;
+	virtual FString GetUsername() const = 0;
+	virtual bool IsOIDCTokenEnabled() = 0;
+	virtual bool IsTokenReady() const = 0;
+	virtual const FString& GetToken() const = 0;
+	virtual UE::Tasks::TTask<void> QueueWorkForToken(TFunction<void(const FString&)> InFunction) = 0;
+
+	virtual void InvalidateCredentials() = 0;
+	virtual void SetLogin(const FString& InUsername, const FString& InPassword) = 0;
+	virtual void SetJiraToken(FJiraToken&& InJiraToken) = 0;
+	virtual const TUniquePtr<FJiraToken>& GetJiraToken() const = 0;
+};
+
+Expose_TNameOf(ICredentialsService);
diff --git a/Engine/Source/Developer/SubmitToolCore/Public/Logic/Services/Interfaces/ILockdownService.h b/Engine/Source/Developer/SubmitToolCore/Public/Logic/Services/Interfaces/ILockdownService.h
new file mode 100644
--- /dev/null
+++ b/Engine/Source/Developer/SubmitToolCore/Public/Logic/Services/Interfaces/ILockdownService.h
@@ -0,0 +1,17 @@
+// Copyright Epic Games, Inc. All Rights Reserved.
+
+#pragma once
+
+#include "CoreMinimal.h"
+#include "Logic/Services/Interfaces/ISubmitToolService.h"
+#include "Models/LockdownData.h"
+#include "Models/SCFile.h"
+
+class ILockdownService : public ISubmitToolService
+{
+public:
+	virtual FSubmitToolLockdownData ArePathsInLockdown(const TArray<FSCFileRef>& InPaths, const TSet<FString>& InDynamicFiles = TSet<FString>()) = 0;
+	virtual bool IsBlockingOperationRunning() const = 0;
+};
+
+Expose_TNameOf(ILockdownService);
diff --git a/Engine/Source/Developer/SubmitToolCore/Public/Logic/Services/Interfaces/ISTSourceControlService.h b/Engine/Source/Developer/SubmitToolCore/Public/Logic/Services/Interfaces/ISTSourceControlService.h
new file mode 100644
--- /dev/null
+++ b/Engine/Source/Developer/SubmitToolCore/Public/Logic/Services/Interfaces/ISTSourceControlService.h
@@ -0,0 +1,200 @@
+// Copyright Epic Games, Inc. All Rights Reserved.
+
+#pragma once
+
+#include "CoreMinimal.h"
+#include "Logic/Services/Interfaces/ISubmitToolService.h"
+#include "Parameters/SubmitToolParameters.h"
+#include "Memory/SharedBuffer.h"
+#include "Tasks/Task.h"
+#include "Misc/Base64.h"
+
+struct FAuthTicket
+{
+public:
+	FAuthTicket() = default;
+	FAuthTicket(const FString& InUsername, const FString& InPassword) : Username(InUsername), Password(InPassword) {}
+	FAuthTicket(FStringView TicketString)
+	{
+		int32 ChopIndex;
+		if(TicketString.FindChar(':', ChopIndex))
+		{
+			Username = TicketString.Left(ChopIndex);
+			Password = TicketString.RightChop(ChopIndex + 1);
+		}
+	}
+	virtual ~FAuthTicket() {}
+
+	virtual FString ToString() const
+	{
+		return TEXT("Basic ") + FBase64::Encode(Username + TEXT(":") + Password);
+	}
+
+	bool IsValid() const
+	{
+		return !Username.IsEmpty() && !Password.IsEmpty();
+	}
+	FString Username;
+private:
+	FString Password;
+};
+
+struct FUserData
+{
+	FUserData(const FString& Username, const FString& Name, const FString& Email) : Name(Name), Username(Username), Email(Email)
+	{}
+
+	FString Name;
+	FString Username;
+	FString Email;
+
+	bool Contains(const FString& SubString) const
+	{
+		return Name.Contains(SubString, ESearchCase::IgnoreCase) ||
+			Username.Contains(SubString, ESearchCase::IgnoreCase) ||
+			Email.Contains(SubString, ESearchCase::IgnoreCase);
+	}
+
+	FString GetDisplayText()
+	{
+		return Name + " - " + Username + " - " + Email;
+	}
+};
+
+struct FSCResultInfo
+{
+	/** Append any messages from another FSCResultInfo, ensuring to keep any already accumulated info. */
+	void Append(const FSCResultInfo& InResultInfo)
+	{
+		InfoMessages.Append(InResultInfo.InfoMessages);
+		ErrorMessages.Append(InResultInfo.ErrorMessages);
+	}
+
+	bool HasErrors() const
+	{
+		return !ErrorMessages.IsEmpty();
+	}
+
+	/** Info and/or warning message storage */
+	TArray<FText> InfoMessages;
+
+	/** Potential error message storage */
+	TArray<FText> ErrorMessages;
+};
+
+
+struct FSCCStream
+{
+public:
+	FSCCStream(const FString& InName, const FString& InParent, const FString& InType)
+		: Name(InName), Parent(InParent), Type(InType)
+	{}
+	const FString Name;
+	const FString Parent;
+	const FString Type;
+
+	TArray<const FString> AdditionalImportPaths;
+};
+
+using FSCCRecord = TMap<FString, FString>;
+using FSCCRecordSet = TArray<FSCCRecord>;
+
+struct FSCCommand
+{
+	// in
+	FString Command;
+	TArray<FString> Args;
+	bool bCollectData = false;
+	bool bSilentErrors = false;
+	FString InputData;
+	// out
+	FSCCRecordSet Values;
+	FSCResultInfo Results;
+	TMap<FString, FSharedBuffer> Data;
+	bool bSuccess = false;
+	bool bDroppedConnection = false;
+};
+
+DECLARE_MULTICAST_DELEGATE_OneParam(FOnUsersGet, TArray<TSharedPtr<FUserData>>& /*Users*/)
+DECLARE_MULTICAST_DELEGATE_OneParam(FOnGroupsGet, TArray<TSharedPtr<FString>>& /*Groups*/)
+DECLARE_MULTICAST_DELEGATE_TwoParams(FOnUsersAndGroupsGet, TArray<TSharedPtr<FUserData>>& /*Users*/, TArray<TSharedPtr<FString>>& /*Groups*/)
+
+using FOnSCCCommandComplete = TDelegate<void(const FSCCommand&)>;
+
+class ISTSourceControlService : public ISubmitToolService
+{
+public:
+	virtual void InitializeParameters(const FSubmitToolParameters& InParameters) = 0;
+	virtual void GetUsers(const FOnUsersGet::FDelegate& Callback) = 0;
+	virtual void GetGroups(const FOnGroupsGet::FDelegate& Callback) = 0;
+	virtual void GetUsersAndGroups(const FOnUsersAndGroupsGet::FDelegate& Callback) = 0;
+
+	virtual UE::Tasks::TTask<bool> DownloadFiles(TArray<FString>&& FilePaths, TMap<FString, FSharedBuffer>& OutFileBuffers, bool bSilentErrors = false) = 0;
+	virtual const TArray<TSharedPtr<FUserData>>& GetRecentUsers() const = 0;
+	virtual void AddRecentUser(TSharedPtr<FUserData>& User) = 0;
+	virtual const TArray<TSharedPtr<FString>>& GetRecentGroups() const = 0;
+	virtual void AddRecentGroup(TSharedPtr<FString>& Group) = 0;
+
+	virtual TSharedPtr<FUserData> GetUserDataFromCache(const FString& Username) const = 0;
+
+	virtual const TArray<FSCCStream*>& GetClientStreams() const = 0;
+	virtual const FSCCStream* GetSCCStream(const FString& InStreamName) = 0;
+	virtual const FString GetRootStreamName() = 0;
+	virtual const FString GetCurrentStreamName() = 0;
+	virtual int32 GetDepotStreamDepth(const FString& InDepotName) = 0;
+	virtual const FAuthTicket& GetAuthTicket() = 0;
+	virtual void RequestCancel() = 0;
+	virtual bool IsBusy() const = 0;
+
+	void RunCommand(FSCCommand&& InCommand, FOnSCCCommandComplete InCompleteCallback)
+	{
+		check(InCompleteCallback.IsBound());
+		TSharedRef<FSCCommand> Command = MakeShared<FSCCommand>(MoveTemp(InCommand));
+
+		UE::Tasks::FTask CommandTask = UE::Tasks::Launch(UE_SOURCE_LOCATION, [this, Command, InCompleteCallback]
+		{
+			Command->bSuccess = RunCommandInternal(*Command);
+			UE::Tasks::Launch(UE_SOURCE_LOCATION, [InCompleteCallback, Command]
+			{
+				TRACE_CPUPROFILER_EVENT_SCOPE_TEXT(*FString::Printf(TEXT("RunP4Command_%s_Callback"), *Command->Command));
+				InCompleteCallback.Execute(*Command);
+			}, LowLevelTasks::ETaskPriority::Normal, UE::Tasks::EExtendedTaskPriority::GameThreadNormalPri);
+		});
+	}
+
+	// Runs the command in a task thread and the provided callback in a dependent GT task.
+	void RunCommand(FString&& InCommand, TArray<FString>&& InAdditionalArgs, FOnSCCCommandComplete InCompleteCallback, bool bSilentErrors = false)
+	{
+		check(InCompleteCallback.IsBound());
+		RunCommand(FSCCommand(MoveTemp(InCommand), MoveTemp(InAdditionalArgs), false, bSilentErrors), InCompleteCallback);
+	}
+
+	// Runs the command in a task thread, calling GetResult() will wait until the task is completed and return the command data.
+	UE::Tasks::TTask<TSharedRef<FSCCommand>> RunCommand(FString&& InCommand, TArray<FString>&& InAdditionalArgs, bool bSilentErrors = false)
+	{
+		TSharedRef<FSCCommand> Command = MakeShared<FSCCommand>(MoveTemp(InCommand), MoveTemp(InAdditionalArgs), false, bSilentErrors);
+		return UE::Tasks::Launch(UE_SOURCE_LOCATION, [this, Command, bSilentErrors]
+		{
+			Command->bSuccess = RunCommandInternal(*Command);
+			return Command;
+		});
+	}
+
+	// Runs the command synchronously and returns the command data
+	FSCCommand RunSyncCommand(FString&& InCommand, TArray<FString>&& InAdditionalArgs, bool bSilentErrors = false)
+	{
+		FSCCommand Command(MoveTemp(InCommand), MoveTemp(InAdditionalArgs), false, bSilentErrors);
+		Command.bSuccess = RunCommandInternal(Command);
+		return MoveTemp(Command);
+	}
+
+	void RunSyncCommand(FSCCommand& InCommand)
+	{
+		InCommand.bSuccess = RunCommandInternal(InCommand);
+	}
+protected:
+	virtual bool RunCommandInternal(FSCCommand& InCommand) = 0;
+
+};
+
+Expose_TNameOf(ISTSourceControlService);
diff --git a/Engine/Source/Developer/SubmitToolCore/Public/Logic/Services/Interfaces/ISubmitToolExtensions.h b/Engine/Source/Developer/SubmitToolCore/Public/Logic/Services/Interfaces/ISubmitToolExtensions.h
new file mode 100644
--- /dev/null
+++ b/Engine/Source/Developer/SubmitToolCore/Public/Logic/Services/Interfaces/ISubmitToolExtensions.h
@@ -0,0 +1,22 @@
+// Copyright Epic Games, Inc. All Rights Reserved.
+
+#pragma once
+
+#include "CoreMinimal.h"
+#include "Logic/Services/Interfaces/ISubmitToolService.h"
+#include "Models/PreflightData.h"
+#include "Models/SCFile.h"
+#include "Parameters/SubmitToolParameters.h"
+
+class FSubmitToolServiceProvider;
+
+class ISubmitToolExtensions : public ISubmitToolService
+{
+public:
+	virtual TOptional<FPreflightTemplateDefinition> SelectPreflightTemplate(
+		const TArray<FSCFileRef>& InFilesInCL,
+		const TSharedPtr<FSubmitToolServiceProvider>& InServiceProvider,
+		const FHordeParameters& InHordeParams) const = 0;
+};
+
+Expose_TNameOf(ISubmitToolExtensions);
diff --git a/Engine/Source/Developer/SubmitToolCore/Public/Logic/Services/Interfaces/ISubmitToolService.h b/Engine/Source/Developer/SubmitToolCore/Public/Logic/Services/Interfaces/ISubmitToolService.h
new file mode 100644
--- /dev/null
+++ b/Engine/Source/Developer/SubmitToolCore/Public/Logic/Services/Interfaces/ISubmitToolService.h
@@ -0,0 +1,11 @@
+// Copyright Epic Games, Inc. All Rights Reserved.
+
+#pragma once
+
+#include "CoreMinimal.h"
+
+class ISubmitToolService
+{
+public:
+	virtual ~ISubmitToolService() = default;
+};
diff --git a/Engine/Source/Developer/SubmitToolCore/Public/Logic/Services/Interfaces/ITagService.h b/Engine/Source/Developer/SubmitToolCore/Public/Logic/Services/Interfaces/ITagService.h
new file mode 100644
--- /dev/null
+++ b/Engine/Source/Developer/SubmitToolCore/Public/Logic/Services/Interfaces/ITagService.h
@@ -0,0 +1,31 @@
+// Copyright Epic Games, Inc. All Rights Reserved.
+
+#pragma once
+
+#include "CoreMinimal.h"
+#include "Logic/Services/Interfaces/ISubmitToolService.h"
+#include "Models/Tag.h"
+
+class ITagService : public ISubmitToolService
+{
+public:
+	virtual FTag* GetTag(const FString& InTagID) = 0;
+	virtual FTag* GetTagOfType(const FString& InType) = 0;
+	virtual FTag* GetTagOfSubtype(const FString& InTagSubtype) = 0;
+	virtual const TArray<const FTag*>& GetTagsArray() const = 0;
+
+	virtual void ParseCLDescription() = 0;
+	virtual void ReApplyTags() = 0;
+	virtual void ApplyTag(const FString& InTagID) = 0;
+	virtual void ApplyTag(FTag& InTag) = 0;
+	virtual void UpdateTagsInCL() = 0;
+	virtual void RemoveTag(const FString& InTagID) = 0;
+	virtual void RemoveTag(FTag& InTag) = 0;
+	virtual void SetTagValues(const FString& InTagID, const FString& InValues) = 0;
+	virtual void SetTagValues(FTag& InTag, const FString& InValues) = 0;
+	virtual void SetTagValues(FTag& InTag, const TArray<FString>& InValues) = 0;
+
+	FTagUpdated OnTagUpdated;
+};
+
+Expose_TNameOf(ITagService);
diff --git a/Engine/Source/Developer/SubmitToolCore/Public/Logic/Services/SubmitToolServiceProvider.h b/Engine/Source/Developer/SubmitToolCore/Public/Logic/Services/SubmitToolServiceProvider.h
new file mode 100644
--- /dev/null
+++ b/Engine/Source/Developer/SubmitToolCore/Public/Logic/Services/SubmitToolServiceProvider.h
@@ -0,0 +1,73 @@
+// Copyright Epic Games, Inc. All Rights Reserved.
+
+#pragma once
+
+#include "CoreMinimal.h"
+#include "Logic/Services/Interfaces/ISubmitToolService.h"
+
+template<class E>
+concept DerivedFromISubmitToolService = std::is_base_of<ISubmitToolService, E>::value;
+
+class FSubmitToolServiceProvider;
+DECLARE_MULTICAST_DELEGATE_OneParam(FOnProviderReady, FSubmitToolServiceProvider& /*ServiceProvider*/)
+
+class FSubmitToolServiceProvider final
+{
+
+public:
+	FSubmitToolServiceProvider()
+	{
+		OnProviderReady.Broadcast(*this);
+	}
+
+	inline static FOnProviderReady OnProviderReady;
+
+	template<DerivedFromISubmitToolService T>
+	TSharedPtr<T> GetService()
+	{
+		FString TypeIdx = TNameOf<T>::GetName();
+		if(Services.Contains(TypeIdx))
+		{
+			return StaticCastSharedPtr<T>(Services[TypeIdx]);
+		}
+
+		return nullptr;
+	}
+
+	template<DerivedFromISubmitToolService T>
+	bool HasService() const
+	{
+		return Services.Contains(TNameOf<T>::GetName());
+	}
+
+	template<DerivedFromISubmitToolService T>
+	TSharedPtr<T> GetService(const FString& InName)
+	{
+		if (Services.Contains(InName))
+		{
+			return StaticCastSharedPtr<T>(Services[InName]);
+		}
+
+		return nullptr;
+	}
+
+	template<DerivedFromISubmitToolService T>
+	void RegisterService(TSharedRef<T> InService)
+	{
+		TSharedPtr<ISubmitToolService>& Service = Services.FindOrAdd(TNameOf<T>::GetName());
+		ensureAlwaysMsgf(Service == nullptr, TEXT("Service %s was re-registered"), TNameOf<T>::GetName());
+		Service = StaticCastSharedPtr<ISubmitToolService>(InService.ToSharedPtr());
+	}
+
+
+	template<DerivedFromISubmitToolService T>
+	void RegisterService(TSharedRef<T> InService, const FString& InName)
+	{
+		TSharedPtr<ISubmitToolService>& Service = Services.FindOrAdd(InName);
+		ensureAlwaysMsgf(Service == nullptr, TEXT("Service %s was re-registered"), *InName);
+		Service = StaticCastSharedPtr<ISubmitToolService>(InService.ToSharedPtr());
+	}
+
+private:
+	TMap<FString, TSharedPtr<ISubmitToolService>> Services;
+};
diff --git a/Engine/Source/Developer/SubmitToolCore/Public/Models/JiraToken.h b/Engine/Source/Developer/SubmitToolCore/Public/Models/JiraToken.h
new file mode 100644
--- /dev/null
+++ b/Engine/Source/Developer/SubmitToolCore/Public/Models/JiraToken.h
@@ -0,0 +1,20 @@
+// Copyright Epic Games, Inc. All Rights Reserved.
+
+#pragma once
+
+#include "CoreMinimal.h"
+
+#include "JiraToken.generated.h"
+
+USTRUCT()
+struct FJiraToken
+{
+	GENERATED_BODY()
+
+	UPROPERTY()
+	FString Access_Token;
+	UPROPERTY()
+	int32 Expires_In = 0;
+	UPROPERTY()
+	FString Refresh_Token;
+};
diff --git a/Engine/Source/Developer/SubmitToolCore/Public/Models/LockdownData.h b/Engine/Source/Developer/SubmitToolCore/Public/Models/LockdownData.h
new file mode 100644
--- /dev/null
+++ b/Engine/Source/Developer/SubmitToolCore/Public/Models/LockdownData.h
@@ -0,0 +1,62 @@
+// Copyright Epic Games, Inc. All Rights Reserved.
+
+#pragma once
+
+#include "CoreMinimal.h"
+#include "Hash/Fnv.h"
+#include "LockdownData.generated.h"
+
+UENUM()
+enum class ELockdownType : int32
+{
+	None = 0,
+	Controlled,
+	Hardcore
+};
+
+struct FCaseInsensitiveStringKeyFuncs : DefaultKeyFuncs<FString, false>
+{
+	static FORCEINLINE bool Matches(const FString& A, const FString& B)
+	{
+		return A.Equals(B, ESearchCase::IgnoreCase);
+	}
+	static FORCEINLINE uint32 GetKeyHash(const FString& Key)
+	{
+		const FString Lower = Key.ToLower();
+		return UE::HashStringFNV1a32(MakeStringView(*Lower));
+	}
+};
+
+struct FSubmitToolLockdownData
+{
+	TSet<FString, FCaseInsensitiveStringKeyFuncs> GroupNames;
+	TSet<FString, FCaseInsensitiveStringKeyFuncs> GroupMessages;
+	ELockdownType LockdownType = ELockdownType::None;
+	TMap<FString, TSet<FString, FCaseInsensitiveStringKeyFuncs>> RequiredTags;
+	bool bIsAllowlisted = true;
+
+	void Append(FSubmitToolLockdownData&& OtherLockdown)
+	{
+		if (OtherLockdown.LockdownType == ELockdownType::None)
+		{
+			return;
+		}
+
+		GroupNames.Append(MoveTemp(OtherLockdown.GroupNames));
+		GroupMessages.Append(MoveTemp(OtherLockdown.GroupMessages));
+		for (auto It = OtherLockdown.RequiredTags.CreateIterator(); It; ++It)
+		{
+			if (RequiredTags.Contains(It->Key))
+			{
+				RequiredTags[It->Key].Append(MoveTemp(It->Value));
+				It.RemoveCurrent();
+			}
+		}
+		RequiredTags.Append(MoveTemp(OtherLockdown.RequiredTags));
+		bIsAllowlisted &= OtherLockdown.bIsAllowlisted;
+		if (LockdownType < OtherLockdown.LockdownType)
+		{
+			LockdownType = OtherLockdown.LockdownType;
+		}
+	}
+};
diff --git a/Engine/Source/Developer/SubmitToolCore/Public/Models/PreflightData.h b/Engine/Source/Developer/SubmitToolCore/Public/Models/PreflightData.h
--- a/Engine/Source/Developer/SubmitToolCore/Public/Models/PreflightData.h
+++ b/Engine/Source/Developer/SubmitToolCore/Public/Models/PreflightData.h
@@ -3,7 +3,6 @@
 #pragma once
 
 #include "CoreMinimal.h"
-#include "Containers/ObservableArray.h"
 #include "PreflightData.generated.h"
 
 UENUM()
diff --git a/Engine/Source/Programs/SubmitTool/Private/Logic/ChangelistService.h b/Engine/Source/Programs/SubmitTool/Private/Logic/ChangelistService.h
--- a/Engine/Source/Programs/SubmitTool/Private/Logic/ChangelistService.h
+++ b/Engine/Source/Programs/SubmitTool/Private/Logic/ChangelistService.h
@@ -4,9 +4,9 @@
 
 #include "CoreMinimal.h"
 #include "Logic/Services/Interfaces/ISTSourceControlService.h"
+#include "Logic/Services/Interfaces/IChangelistService.h"
 #include "Containers/Ticker.h"
 #include "Parameters/SubmitToolParameters.h"
-#include "Services/Interfaces/ISubmitToolService.h"
 #include "Tasks/Task.h"
 
 class FSubmitToolServiceProvider;
@@ -14,9 +14,6 @@ using FSourceControlOperationRef = TSharedRef<class ISourceControlOperation, ESP
 
 
 DECLARE_DELEGATE(FOnP4OperationCompleteDelegate)
-DECLARE_DELEGATE_OneParam(FOnChangeListReadyDelegate, bool /*bValidState*/)
-DECLARE_DELEGATE_OneParam(FOnChangelistRefreshDelegate, ETaskArea /*ChangedValue*/)
-DECLARE_MULTICAST_DELEGATE(FOnCLDescriptionUpdated)
 
 struct FChangelistData
 {
@@ -34,14 +31,14 @@ struct FChangelistData
 		const TMap<FString, FString>& DepotFileToLocalFile);
 };
 
-class FChangelistService final : public ISubmitToolService
+class FChangelistService final : public IChangelistService
 {
 public:
 
 	FChangelistService(const FGeneralParameters& InParameters, const TSharedPtr<FSubmitToolServiceProvider> InServiceProvider, const FOnChangeListReadyDelegate& InCLReadyCallback, const FOnChangelistRefreshDelegate& InCLRefreshCallback);
 	~FChangelistService();
 
-	const FString GetCLID() const
+	virtual const FString GetCLID() const override
 	{
 		check(IsInGameThread());
 		if(!CLID.IsEmpty())
@@ -52,13 +49,13 @@ public:
 		return TEXT("Invalid");
 	}
 
-	const FString& GetCLDescription()
+	virtual const FString& GetCLDescription() override
 	{
 		check(IsInGameThread());
 		return CLDescription;
 	}
 
-	bool SetCLDescription(const FString& newDescription, bool bNotifyEvent = false)
+	virtual bool SetCLDescription(const FString& newDescription, bool bNotifyEvent = false) override
 	{
 		check(IsInGameThread());
 		FString lineEndReplaced = newDescription.Replace(TEXT("\r\n"), TEXT("\n"));
@@ -78,68 +75,66 @@ public:
 		return true;
 	}
 
-	const TArray<FSCFileRef>& GetFilesInCL() const
+	virtual const TArray<FSCFileRef>& GetFilesInCL() const override
 	{
 		check(IsInGameThread());
 		return Changelist.Files;
 	}
 
-	const TArray<FSCFileRef>& GetShelvedFilesInCL() const
+	virtual const TArray<FSCFileRef>& GetShelvedFilesInCL() const override
 	{
 		check(IsInGameThread());
 		return Changelist.ShelvedFiles;
 	}
 
-	bool HasShelvedFiles() const
+	virtual bool HasShelvedFiles() const override
 	{
 		check(IsInGameThread());
 		return Changelist.ShelvedFiles.Num() > 0;
 	}
 
-	bool IsShelfComplete() const
+	virtual bool IsShelfComplete() const override
 	{
 		check(IsInGameThread());
 		return Changelist.bIsShelfComplete;
 	}
 
-	const TMap<FString, FString>& GetOpenFilesInChangelists() const
+	virtual const TMap<FString, FString>& GetOpenFilesInChangelists() const override
 	{
 		check(IsInGameThread());
 		return Changelist.AllOpenFiles;
 	}
 
-	bool HasP4OperationsRunning() const
+	virtual bool HasP4OperationsRunning() const override
 	{
 		check(IsInGameThread());
 		return !UpdateChangelistTask.IsCompleted() || SourceControlService->IsBusy();
 	}
 
-	bool IsShelfReady() const
+	virtual bool IsShelfReady() const override
 	{
 		return bIsShelfReady;
 	}
 
-	void Init();
-	void Submit(const FString& InDescriptionAddendum = TEXT(""), TFunction<void(const FSCCommand&)> OnComplete = nullptr);
+	virtual void Init() override;
+	virtual void Submit(const FString& InDescriptionAddendum = TEXT(""), TFunction<void(const FSCCommand&)> OnComplete = nullptr) override;
 
 	//void Submit(const FString& InDescriptionAddendum = TEXT(""), const FSourceControlOperationComplete& OnSubmitComplete = FSourceControlOperationComplete());
-	void CreateCLFromDefaultCL();
-	void FetchChangelistDataAsync();
-	void RevertUnchangedFiles(TFunction<void(const FSCCommand&)> OnComplete = nullptr);
-	
-	void UpdateP4CLDescription();
-	void UpdateP4CLDescriptionSynchronously();
-	void DeleteShelvedFiles(TFunction<void(const FSCCommand&)> OnComplete = nullptr);
-	void ReplaceShelvedFiles(TFunction<void(const FSCCommand&)> OnComplete = nullptr);
-	void CreateShelvedFiles(TFunction<void(const FSCCommand&)> OnComplete = nullptr);
-	void UpdateShelvedFiles(TArray<FString> InFiles, TFunction<void(const FSCCommand&)> OnComplete = nullptr);
-	void CheckShelfIsUpToDate(TFunction<void(bool, TArray<FString>)> OnComplete);
-	void EnsureShelfIsCurrent();
+	virtual void CreateCLFromDefaultCL() override;
+	virtual void FetchChangelistDataAsync() override;
+	virtual void RevertUnchangedFiles(TFunction<void(const FSCCommand&)> OnComplete = nullptr) override;
+
+	virtual void UpdateP4CLDescription() override;
+	virtual void UpdateP4CLDescriptionSynchronously() override;
+	virtual void DeleteShelvedFiles(TFunction<void(const FSCCommand&)> OnComplete = nullptr) override;
+	virtual void ReplaceShelvedFiles(TFunction<void(const FSCCommand&)> OnComplete = nullptr) override;
+	virtual void CreateShelvedFiles(TFunction<void(const FSCCommand&)> OnComplete = nullptr) override;
+	virtual void UpdateShelvedFiles(TArray<FString> InFiles, TFunction<void(const FSCCommand&)> OnComplete = nullptr) override;
+	virtual void CheckShelfIsUpToDate(TFunction<void(bool, TArray<FString>)> OnComplete) override;
+	virtual void EnsureShelfIsCurrent() override;
 	bool P4Tick(float DeltaTime);
 
-	void CancelP4Operations();
-
-	FOnCLDescriptionUpdated OnCLDescriptionUpdated;
+	virtual void CancelP4Operations() override;
 
 private:
 	mutable FCriticalSection Mutex;
diff --git a/Engine/Source/Programs/SubmitTool/Private/Logic/CredentialsService.h b/Engine/Source/Programs/SubmitTool/Private/Logic/CredentialsService.h
--- a/Engine/Source/Programs/SubmitTool/Private/Logic/CredentialsService.h
+++ b/Engine/Source/Programs/SubmitTool/Private/Logic/CredentialsService.h
@@ -6,9 +6,8 @@
 #include "HttpServerRequest.h"
 #include "Parameters/SubmitToolParameters.h"
 #include "Tasks/Task.h"
-#include "Services/Interfaces/ISubmitToolService.h"
-
-#include "CredentialsService.generated.h"
+#include "Logic/Services/Interfaces/ICredentialsService.h"
+#include "Models/JiraToken.h"
 
 class FProcessWrapper;
 
@@ -39,55 +38,42 @@ private:
 	FHttpRequestHandler Callback;
 };
 
-USTRUCT()
-struct FJiraToken
-{
-	GENERATED_BODY()
-
-	UPROPERTY()
-	FString Access_Token;
-	UPROPERTY()
-	int32 Expires_In;
-	UPROPERTY()
-	FString Refresh_Token;
-};
-
-class FCredentialsService : public ISubmitToolService
+class FCredentialsService : public ICredentialsService
 {
 public:
 	FCredentialsService(const FSubmitToolParameters& InParameters);
 	virtual ~FCredentialsService();
 
-	bool HasJiraCredentials() const
+	virtual bool HasJiraCredentials() const override
 	{
 		return Parameters.JiraParameters.OAuthClientId.IsEmpty() ? !LoginString.IsEmpty() : JiraOAuthToken != nullptr;
 	}
 
-	void InvalidateCredentials()
+	virtual void InvalidateCredentials() override
 	{
 		LoginString.Empty();
 		JiraOAuthToken = nullptr;
 	}
 
-	const FString& GetEncodedLoginString() const
+	virtual const FString& GetEncodedLoginString() const override
 	{
 		return LoginString;
 	}
 
-	FString GetUsername() const;
-	void SetLogin(const FString& InUsername, const FString& InPassword);
+	virtual FString GetUsername() const override;
+	virtual void SetLogin(const FString& InUsername, const FString& InPassword) override;
 
-	bool IsOIDCTokenEnabled()
+	virtual bool IsOIDCTokenEnabled() override
 	{
 		return !Parameters.OAuthParameters.OAuthTokenTool.IsEmpty();
 	}
 
-	bool IsTokenReady() const
+	virtual bool IsTokenReady() const override
 	{
 		return !OIDCToken.IsEmpty();
 	}
 
-	const FString& GetToken() const
+	virtual const FString& GetToken() const override
 	{
 		return OIDCToken;
 	}
@@ -102,15 +88,15 @@ public:
 		return Key;
 	}
 
-	UE::Tasks::TTask<void> QueueWorkForToken(TFunction<void(const FString&)> InFunction);
+	virtual UE::Tasks::TTask<void> QueueWorkForToken(TFunction<void(const FString&)> InFunction) override;
 
-	void SetJiraToken(FJiraToken&& InJiraToken)
+	virtual void SetJiraToken(FJiraToken&& InJiraToken) override
 	{
 		JiraOAuthToken = MakeUnique<FJiraToken>(MoveTemp(InJiraToken));
 		SaveCredentials();
 	}
 
-	const TUniquePtr<FJiraToken>& GetJiraToken() const
+	virtual const TUniquePtr<FJiraToken>& GetJiraToken() const override
 	{
 		return JiraOAuthToken;
 	}
diff --git a/Engine/Source/Programs/SubmitTool/Private/Logic/IntegrationService.cpp b/Engine/Source/Programs/SubmitTool/Private/Logic/IntegrationService.cpp
--- a/Engine/Source/Programs/SubmitTool/Private/Logic/IntegrationService.cpp
+++ b/Engine/Source/Programs/SubmitTool/Private/Logic/IntegrationService.cpp
@@ -7,9 +7,9 @@
 
 #include "CommandLine/CmdLineParameters.h"
 #include "Logging/SubmitToolLog.h"
-#include "Logic/ChangelistService.h"
+#include "Logic/Services/Interfaces/IChangelistService.h"
 #include "Logic/TasksService.h"
-#include "Logic/TagService.h"
+#include "Logic/Services/Interfaces/ITagService.h"
 #include "Logic/JiraService.h"
 #include "Logic/SwarmService.h"
 #include "Logic/Services/Interfaces/ICacheDataService.h"
@@ -84,10 +84,10 @@ void FIntegrationService::RequestIntegration(const FOnBooleanValueChanged OnComp
 	SwarmService->FetchReview(OnGetReviewComplete::CreateLambda(
 	[this, SwarmService, OnComplete](const TUniquePtr<FSwarmReview>& InReview, const FString& InErrorMessage)
 	{
-		TSharedPtr<FChangelistService> ChangelistService = ServiceProvider.Pin()->GetService<FChangelistService>();
+		TSharedPtr<IChangelistService> ChangelistService = ServiceProvider.Pin()->GetService<IChangelistService>();
 		TSharedPtr<ISTSourceControlService> SourceControlService = ServiceProvider.Pin()->GetService<ISTSourceControlService>();
 		TSharedPtr<FJiraService> JiraService = ServiceProvider.Pin()->GetService<FJiraService>();
-		TSharedPtr<FTagService> TagService = ServiceProvider.Pin()->GetService<FTagService>();
+		TSharedPtr<ITagService> TagService = ServiceProvider.Pin()->GetService<ITagService>();
 
 		UE_LOG(LogSubmitTool, Log, TEXT("Requesting Integration..."));
 		const FTag* JiraTag = TagService->GetTag(TEXT("#jira"));
@@ -102,7 +102,7 @@ void FIntegrationService::RequestIntegration(const FOnBooleanValueChanged OnComp
 		{
 			TFunction<void()> CreateSwarm = [this, JiraTagValue, OnComplete]() {
 				TSharedPtr<FSwarmService> SwarmService = ServiceProvider.Pin()->GetService<FSwarmService>();
-				TSharedPtr<FChangelistService> ChangelistService = ServiceProvider.Pin()->GetService<FChangelistService>();
+				TSharedPtr<IChangelistService> ChangelistService = ServiceProvider.Pin()->GetService<IChangelistService>();
 				TSharedPtr<ISTSourceControlService> SourceControlService = ServiceProvider.Pin()->GetService<ISTSourceControlService>();
 				TSharedPtr<FJiraService> JiraService = ServiceProvider.Pin()->GetService<FJiraService>();
 				FString SwarmUrl;
@@ -203,7 +203,7 @@ bool FIntegrationService::ValidateIntegrationOptions(bool bSilent) const
 {
 	bool bGlobalValidation = true;
 
-	const FTag* JiraTag = ServiceProvider.Pin()->GetService<FTagService>()->GetTag(TEXT("#jira"));
+	const FTag* JiraTag = ServiceProvider.Pin()->GetService<ITagService>()->GetTag(TEXT("#jira"));
 	if(JiraTag != nullptr && JiraTag->GetValues().Num() > 0)
 	{
 		bool bPassesRegex = true;
diff --git a/Engine/Source/Programs/SubmitTool/Private/Logic/IntegrationService.h b/Engine/Source/Programs/SubmitTool/Private/Logic/IntegrationService.h
--- a/Engine/Source/Programs/SubmitTool/Private/Logic/IntegrationService.h
+++ b/Engine/Source/Programs/SubmitTool/Private/Logic/IntegrationService.h
@@ -6,7 +6,7 @@
 
 #include "Parameters/SubmitToolParameters.h"
 #include "Framework/SlateDelegates.h"
-#include "Services/Interfaces/ISubmitToolService.h"
+#include "Logic/Services/Interfaces/ISubmitToolService.h"
 
 class FTasksService;
 class FChangelistService;
diff --git a/Engine/Source/Programs/SubmitTool/Private/Logic/JiraService.cpp b/Engine/Source/Programs/SubmitTool/Private/Logic/JiraService.cpp
--- a/Engine/Source/Programs/SubmitTool/Private/Logic/JiraService.cpp
+++ b/Engine/Source/Programs/SubmitTool/Private/Logic/JiraService.cpp
@@ -5,7 +5,7 @@
 #include "Logic/DialogFactory.h"
 #include "Logging/SubmitToolLog.h"
 #include "SubmitToolCoreUtils.h"
-#include "Logic/TagService.h"
+#include "Logic/Services/Interfaces/ITagService.h"
 #include "Logic/PreflightService.h"
 #include "Logic/Services/SubmitToolServiceProvider.h"
 #include "Logic/Services/Interfaces/ISTSourceControlService.h"
@@ -27,7 +27,7 @@ FJiraService::FJiraService(const FJiraParameters& InJiraSettings, const int32 In
 	if(!Definition.ServerAddress.IsEmpty())
 	{
 		LoadJiraIssues();
-		if (ServiceProvider.Pin()->GetService<FCredentialsService>()->HasJiraCredentials())
+		if (ServiceProvider.Pin()->GetService<ICredentialsService>()->HasJiraCredentials())
 		{
 			FetchJiraTickets(true, true);
 		}
@@ -108,7 +108,7 @@ void FJiraService::ObtainOAuth(TFunction<void(bool)> InCallback)
 
 bool FJiraService::FetchJiraTickets(bool InForce, bool bInFailSilently)
 {
-	TSharedPtr<FCredentialsService> CredentialsService = ServiceProvider.Pin()->GetService<FCredentialsService>();
+	TSharedPtr<ICredentialsService> CredentialsService = ServiceProvider.Pin()->GetService<ICredentialsService>();
 	if(CredentialsService->HasJiraCredentials() && !Definition.ServerAddress.IsEmpty() && !Definition.bDisableJira)
 	{
 		if (InForce || JiraIssues.Num() == 0)
@@ -128,7 +128,7 @@ void FJiraService::Reset()
 
 void FJiraService::QueryIssues(bool bInFailSilently, bool bAllowTokenRefresh)
 {
-	TSharedPtr<FCredentialsService> CredentialsService = ServiceProvider.Pin()->GetService<FCredentialsService>();
+	TSharedPtr<ICredentialsService> CredentialsService = ServiceProvider.Pin()->GetService<ICredentialsService>();
 	if(JiraRequest.IsValid() || Definition.bDisableJira || !CredentialsService->HasJiraCredentials())
 	{
 		return;
@@ -162,7 +162,7 @@ void FJiraService::QueryIssues_HttpRequestComplete(FHttpRequestPtr HttpRequest,
 {
 	JiraRequest = nullptr;
 	bOngoingRequest = false;
-	TSharedPtr<FCredentialsService> CredentialsService = ServiceProvider.Pin()->GetService<FCredentialsService>();
+	TSharedPtr<ICredentialsService> CredentialsService = ServiceProvider.Pin()->GetService<ICredentialsService>();
 
 	if (!bSucceeded)
 	{
@@ -239,7 +239,7 @@ void FJiraService::QueryIssues_HttpRequestComplete(FHttpRequestPtr HttpRequest,
 					else
 					{
 						UE_LOG(LogSubmitTool, Error, TEXT("Jira token refresh failed, re-authentication required."));
-						ServiceProvider.Pin()->GetService<FCredentialsService>()->InvalidateCredentials();
+						ServiceProvider.Pin()->GetService<ICredentialsService>()->InvalidateCredentials();
 						OnJiraIssuesRetrievedCallback.ExecuteIfBound(false);
 					}
 				});
@@ -359,7 +359,7 @@ void FJiraService::SaveJiraIssues() const
 void FJiraService::LoadJiraIssues()
 {
 	// do not load the issues if there is no credentials
-	if (!ServiceProvider.Pin()->GetService<FCredentialsService>()->HasJiraCredentials())
+	if (!ServiceProvider.Pin()->GetService<ICredentialsService>()->HasJiraCredentials())
 	{
 		return;
 	}
@@ -457,7 +457,7 @@ void FJiraService::ExchangeOAuthCode(TFunction<void(bool)> InCallback)
 				FJiraToken Token;
 				if (FJsonObjectConverter::JsonObjectStringToUStruct<FJiraToken>(*ResponseStr, &Token))
 				{
-					ServiceProvider.Pin()->GetService<FCredentialsService>()->SetJiraToken(MoveTemp(Token));
+					ServiceProvider.Pin()->GetService<ICredentialsService>()->SetJiraToken(MoveTemp(Token));
 					UE_LOG(LogSubmitToolDebug, Log, TEXT("Obtained Jira Access Token correctly."));
 					InCallback(true);
 					return;
@@ -488,7 +488,7 @@ void FJiraService::ExchangeOAuthCode(TFunction<void(bool)> InCallback)
 
 void FJiraService::RefreshOAuthToken(TFunction<void(bool)> InCallback)
 {
-	TSharedPtr<FCredentialsService> CredentialsService = ServiceProvider.Pin()->GetService<FCredentialsService>();
+	TSharedPtr<ICredentialsService> CredentialsService = ServiceProvider.Pin()->GetService<ICredentialsService>();
 	const TUniquePtr<FJiraToken>& Token = CredentialsService->GetJiraToken();
 
 	if (!Token.IsValid() || Token->Refresh_Token.IsEmpty())
@@ -532,7 +532,7 @@ void FJiraService::RefreshOAuthToken(TFunction<void(bool)> InCallback)
 					FJiraToken NewToken;
 					if (FJsonObjectConverter::JsonObjectStringToUStruct<FJiraToken>(*HttpResponse->GetContentAsString(), &NewToken))
 					{
-						ServiceProvider.Pin()->GetService<FCredentialsService>()->SetJiraToken(MoveTemp(NewToken));
+						ServiceProvider.Pin()->GetService<ICredentialsService>()->SetJiraToken(MoveTemp(NewToken));
 						UE_LOG(LogSubmitTool, Log, TEXT("Jira access token refreshed successfully."));
 						bSuccess = true;
 					}
@@ -563,7 +563,7 @@ const FString FJiraService::GetJiraIssuesFilepath() const
 
 void FJiraService::GetIssueAndCreateServiceDeskRequest(const FString& Key, const FString& Description, const FString& SwarmURL, const FString& InCurrentStream, const TMap<FString, TSharedPtr<FIntegrationOptionBase>>& InIntegrationOptions, const FOnBooleanValueChanged OnComplete)
 {
-	TSharedPtr<FCredentialsService> CredentialsService = ServiceProvider.Pin()->GetService<FCredentialsService>();
+	TSharedPtr<ICredentialsService> CredentialsService = ServiceProvider.Pin()->GetService<ICredentialsService>();
 
 	if(Key.IsEmpty() || Key.Equals(TEXT("none"), ESearchCase::IgnoreCase) || !CredentialsService->HasJiraCredentials())
 	{
@@ -589,7 +589,7 @@ void FJiraService::GetIssueAndCreateServiceDeskRequest(const FString& Key, const
 		Url = FString::Format(TEXT("https://api.atlassian.com/ex/jira/{0}/rest/api/2/issue/{1}"), { Definition.CloudID, Key });
 		if (CredentialsService->GetJiraToken().IsValid())
 		{
-			HttpRequest->SetHeader(TEXT("Authorization"), TEXT("Bearer ") + ServiceProvider.Pin()->GetService<FCredentialsService>()->GetJiraToken()->Access_Token);
+			HttpRequest->SetHeader(TEXT("Authorization"), TEXT("Bearer ") + ServiceProvider.Pin()->GetService<ICredentialsService>()->GetJiraToken()->Access_Token);
 		}
 	}
 
@@ -687,7 +687,7 @@ void FJiraService::CreateServiceDeskRequest(TSharedPtr<FJsonObject> InBaseJiraJs
 
 	if(!Definition.PreflightField.IsEmpty())
 	{		
-		const FTag* PreflightTag = ServiceProvider.Pin()->GetService<FTagService>()->GetTagOfSubtype(TEXT("preflight"));
+		const FTag* PreflightTag = ServiceProvider.Pin()->GetService<ITagService>()->GetTagOfSubtype(TEXT("preflight"));
 		if (PreflightTag != nullptr && PreflightTag->GetValues().Num() > 0)
 		{
 			FString PreflightTagValue = PreflightTag->GetValues()[0];
@@ -708,7 +708,7 @@ void FJiraService::CreateServiceDeskRequest(TSharedPtr<FJsonObject> InBaseJiraJs
 
 	if (!Definition.JiraField.IsEmpty())
 	{
-		const FTag* JiraTag = ServiceProvider.Pin()->GetService<FTagService>()->GetTagOfType(TEXT("JiraIssue"));
+		const FTag* JiraTag = ServiceProvider.Pin()->GetService<ITagService>()->GetTagOfType(TEXT("JiraIssue"));
 		if (JiraTag != nullptr && JiraTag->GetValues().Num() > 0)
 		{
 			TArray<FString> URLs;
diff --git a/Engine/Source/Programs/SubmitTool/Private/Logic/JiraService.h b/Engine/Source/Programs/SubmitTool/Private/Logic/JiraService.h
--- a/Engine/Source/Programs/SubmitTool/Private/Logic/JiraService.h
+++ b/Engine/Source/Programs/SubmitTool/Private/Logic/JiraService.h
@@ -8,7 +8,7 @@
 #include "Models/IntegrationOptions.h"
 #include "Interfaces/IHttpRequest.h"
 #include "Framework/SlateDelegates.h"
-#include "Services/Interfaces/ISubmitToolService.h"
+#include "Logic/Services/Interfaces/ISubmitToolService.h"
 
 DECLARE_DELEGATE_OneParam(FOnJiraIssuesRetrieved, bool /*bValidResponse*/)
 
diff --git a/Engine/Source/Programs/SubmitTool/Private/Logic/P4LockdownService.h b/Engine/Source/Programs/SubmitTool/Private/Logic/P4LockdownService.h
--- a/Engine/Source/Programs/SubmitTool/Private/Logic/P4LockdownService.h
+++ b/Engine/Source/Programs/SubmitTool/Private/Logic/P4LockdownService.h
@@ -6,81 +6,22 @@
 #include "Parameters/SubmitToolParameters.h"
 #include "Logic/ProcessWrapper.h"
 #include "Tasks/Task.h"
-#include "Services/Interfaces/ISubmitToolService.h"
+#include "Logic/Services/Interfaces/ILockdownService.h"
 #include "Memory/SharedBuffer.h"
-#include "Hash/Fnv.h"
-#include "P4LockdownService.generated.h"
 
 
 DECLARE_DELEGATE_OneParam(FOnConfigFileRetrieved, bool /*bValid*/)
 
 class FSubmitToolServiceProvider;
 
-UENUM()
-enum class ELockdownType : int32
-{
-	None = 0,	
-	Controlled,
-	Hardcore
-};
-
-struct FCaseInsensitiveStringKeyFuncs : DefaultKeyFuncs<FString, false>
-{
-	static FORCEINLINE bool Matches(const FString& A, const FString& B)
-	{
-		return A.Equals(B, ESearchCase::IgnoreCase);
-	}
-	static FORCEINLINE uint32 GetKeyHash(const FString& Key)
-	{
-		const FString Lower = Key.ToLower();
-		return UE::HashStringFNV1a32(MakeStringView(*Lower));
-	}
-};
-
-
-struct FSubmitToolLockdownData
-{
-	TSet<FString, FCaseInsensitiveStringKeyFuncs> GroupNames;
-	TSet<FString, FCaseInsensitiveStringKeyFuncs> GroupMessages;
-	ELockdownType LockdownType = ELockdownType::None;
-	TMap<FString, TSet<FString, FCaseInsensitiveStringKeyFuncs>> RequiredTags;
-	bool bIsAllowlisted = true;
-
-
-	void Append(FSubmitToolLockdownData&& OtherLockdown)
-	{
-		if (OtherLockdown.LockdownType == ELockdownType::None)
-		{
-			return;
-		}
-
-		GroupNames.Append(MoveTemp(OtherLockdown.GroupNames));
-		GroupMessages.Append(MoveTemp(OtherLockdown.GroupMessages));
-		for (auto It = OtherLockdown.RequiredTags.CreateIterator(); It; ++It)
-		{
-			if (RequiredTags.Contains(It->Key))
-			{
-				RequiredTags[It->Key].Append(MoveTemp(It->Value));
-				It.RemoveCurrent();
-			}
-		}
-		RequiredTags.Append(MoveTemp(OtherLockdown.RequiredTags));
-		bIsAllowlisted &= OtherLockdown.bIsAllowlisted;
-		if (LockdownType < OtherLockdown.LockdownType)
-		{
-			LockdownType = OtherLockdown.LockdownType;
-		}
-	}
-};
-
-class FP4LockdownService : public ISubmitToolService
+class FP4LockdownService : public ILockdownService
 {	
 public:
 
 	FP4LockdownService() = delete;
 	FP4LockdownService(const FP4LockdownParameters& InParameters, TWeakPtr<FSubmitToolServiceProvider> InServiceProvider);
-	FSubmitToolLockdownData ArePathsInLockdown(const TArray<FSCFileRef>& InPaths, const TSet<FString>& InDynamicFiles = TSet<FString>());
-	bool IsBlockingOperationRunning() const
+	virtual FSubmitToolLockdownData ArePathsInLockdown(const TArray<FSCFileRef>& InPaths, const TSet<FString>& InDynamicFiles = TSet<FString>()) override;
+	virtual bool IsBlockingOperationRunning() const override
 	{
 		return DownloadTask.IsValid() && !DownloadTask.IsCompleted();
 	}
diff --git a/Engine/Source/Programs/SubmitTool/Private/Logic/PreflightService.cpp b/Engine/Source/Programs/SubmitTool/Private/Logic/PreflightService.cpp
--- a/Engine/Source/Programs/SubmitTool/Private/Logic/PreflightService.cpp
+++ b/Engine/Source/Programs/SubmitTool/Private/Logic/PreflightService.cpp
@@ -15,10 +15,11 @@
 #include "Widgets/Text/STextBlock.h"
 
 #include "CommandLine/CmdLineParameters.h"
+#include "Logic/Services/Interfaces/ISubmitToolExtensions.h"
 #include "Configuration/Configuration.h"
 #include "Logic/Services/Interfaces/ISTSourceControlService.h"
-#include "Logic/CredentialsService.h"
-#include "Logic/TagService.h"
+#include "Logic/Services/Interfaces/ICredentialsService.h"
+#include "Logic/Services/Interfaces/ITagService.h"
 #include "Logging/SubmitToolLog.h"
 #include "Models/ModelInterface.h"
 #include "Models/SubmitToolUserPrefs.h"
@@ -39,7 +40,7 @@ FPreflightService::FPreflightService(
 {
 	TickHandle = FTSTicker::GetCoreTicker().AddTicker(FTickerDelegate::CreateRaw(this, &FPreflightService::Tick));
 
-	PreflightTag = ServiceProvider.Pin()->GetService<FTagService>()->GetTagOfSubtype(TEXT("preflight"));
+	PreflightTag = ServiceProvider.Pin()->GetService<ITagService>()->GetTagOfSubtype(TEXT("preflight"));
 }
 
 FPreflightService::~FPreflightService()
@@ -50,7 +51,7 @@ FPreflightService::~FPreflightService()
 
 bool FPreflightService::Tick(float DeltaTime)
 {
-	TSharedPtr<FChangelistService> ChangelistService = ServiceProvider.Pin()->GetService<FChangelistService>();
+	TSharedPtr<IChangelistService> ChangelistService = ServiceProvider.Pin()->GetService<IChangelistService>();
 	switch(State)
 	{
 		//////////////////////////////////////////////////
@@ -86,7 +87,7 @@ TMap<FString, FStringFormatArg> FPreflightService::GetFormatParameters(const FSt
 	TMap<FString, FStringFormatArg> FormatMap =
 	{
 		{ TEXT("URL"), Definition.HordeServerAddress },
-		{ TEXT("CLID"), ServiceProvider.Pin()->GetService<FChangelistService>()->GetCLID()}
+		{ TEXT("CLID"), ServiceProvider.Pin()->GetService<IChangelistService>()->GetCLID()}
 	};
 
 	FormatMap.Add(TEXT("Stream"), StreamName);
@@ -120,7 +121,7 @@ void FPreflightService::RequestPreflight(const FString InTemplate)
 			TemplateId = InTemplate;
 			UE_LOG(LogSubmitTool, Log, TEXT("Preflight: Requesting..."));
 			
-			TSharedPtr<FChangelistService> ChangelistService = ServiceProvider.Pin()->GetService<FChangelistService>();
+			TSharedPtr<IChangelistService> ChangelistService = ServiceProvider.Pin()->GetService<IChangelistService>();
 			
 			ChangelistService->EnsureShelfIsCurrent();
 			State = EPreflightServiceState::WaitingForShelf;
@@ -173,14 +174,14 @@ void FPreflightService::FetchPreflightInfo(bool bRequeue, const FString& InOAuth
 {
 	TRACE_CPUPROFILER_EVENT_SCOPE(FPreflightService::FetchPreflightInfo);
 
-	TSharedPtr<FCredentialsService> CredentialsService = ServiceProvider.Pin()->GetService<FCredentialsService>();
+	TSharedPtr<ICredentialsService> CredentialsService = ServiceProvider.Pin()->GetService<ICredentialsService>();
 	if(Definition.HordeServerAddress.IsEmpty() || !CredentialsService->IsOIDCTokenEnabled() || FModelInterface::GetState() == ESubmitToolAppState::Finished)
 	{
 		bHordeInformationReady = true;
 		return;
 	}
 
-	TSharedPtr<FChangelistService> ChangelistService = ServiceProvider.Pin()->GetService<FChangelistService>();
+	TSharedPtr<IChangelistService> ChangelistService = ServiceProvider.Pin()->GetService<IChangelistService>();
 
 	// Don't bother with the default changelist
 	if(ChangelistService->GetCLID() == TEXT("default"))
@@ -287,7 +288,7 @@ void FPreflightService::FetchPreflightInfo(bool bRequeue, const FString& InOAuth
 							// Only log when there's a different number of preflights
 							if(!HordePreflights.IsValid() || HordePreflights->PreflightList.Num() != NewHordePreflights->PreflightList.Num())
 							{
-								UE_LOG(LogSubmitTool, Log, TEXT("Retrieved %d preflights for CL %s"), NewHordePreflights->PreflightList.Num(), *ServiceProvider.Pin()->GetService<FChangelistService>()->GetCLID());
+								UE_LOG(LogSubmitTool, Log, TEXT("Retrieved %d preflights for CL %s"), NewHordePreflights->PreflightList.Num(), *ServiceProvider.Pin()->GetService<IChangelistService>()->GetCLID());
 							}
 
 							HordePreflights = MoveTemp(NewHordePreflights);
@@ -507,7 +508,7 @@ void FPreflightService::FetchUnlinkedPreflight(const FString& InPreflightId, boo
 
 void FPreflightService::StartPreflight()
 {
-	const TArray<FSCFileRef>& ShelvedFiles = ServiceProvider.Pin()->GetService<FChangelistService>()->GetShelvedFilesInCL();
+	const TArray<FSCFileRef>& ShelvedFiles = ServiceProvider.Pin()->GetService<IChangelistService>()->GetShelvedFilesInCL();
 
 	if (ShelvedFiles.IsEmpty())
 	{
@@ -595,11 +596,21 @@ TArray<const FPreflightData*> FPreflightService::GetTaggedPreflights() const
 
 bool FPreflightService::SelectPreflightTemplate(FPreflightTemplateDefinition& OutTemplate) const
 {
-	TSharedPtr<FChangelistService> ChangelistService = ServiceProvider.Pin()->GetService<FChangelistService>();
+	TSharedPtr<IChangelistService> ChangelistService = ServiceProvider.Pin()->GetService<IChangelistService>();
 
 	// Get the list of files in the changelist
 	const TArray<FSCFileRef>& FilesInCl = bCheckShelveInstead ? ChangelistService->GetShelvedFilesInCL() : ChangelistService->GetFilesInCL();
 
+	if (ServiceProvider.Pin()->HasService<ISubmitToolExtensions>())
+	{
+		TOptional<FPreflightTemplateDefinition> Result = ServiceProvider.Pin()->GetService<ISubmitToolExtensions>()->SelectPreflightTemplate(FilesInCl, ServiceProvider.Pin(), Definition);
+		if (Result.IsSet())
+		{
+			OutTemplate = Result.GetValue();
+			return true;
+		}
+	}
+
 	// Loop through each definition to see if the files are in the path then check extension
 	for(const FPreflightTemplateDefinition& Def : Definition.Definitions)
 	{
@@ -624,7 +635,7 @@ FString FPreflightService::GetAdditionalTasksString(const FPreflightTemplateDefi
 	TStringBuilder<256> AdditionalTaskStrBuilder;
 	const FString BaseString = TEXT("&id-additional-tasks.");
 	const FString EndString = TEXT("=true");
-	TSharedPtr<FChangelistService> ChangelistService = ServiceProvider.Pin()->GetService<FChangelistService>();
+	TSharedPtr<IChangelistService> ChangelistService = ServiceProvider.Pin()->GetService<IChangelistService>();
 
 	const TArray<FSCFileRef>& FilesInCl = bCheckShelveInstead ? ChangelistService->GetShelvedFilesInCL() : ChangelistService->GetFilesInCL();
 	for(const FPreflightAdditionalTask& AdditionalTask : InTemplate.AdditionalTasks)
@@ -651,7 +662,7 @@ void FPreflightService::RefreshStreamName()
 	TSharedPtr<ISTSourceControlService> SCCService = ServiceProvider.Pin()->GetService<ISTSourceControlService>();
 	StreamName = SCCService->GetRootStreamName();
 
-	const TArray<FSCFileRef>& ShelvedFiles = ServiceProvider.Pin()->GetService<FChangelistService>()->GetShelvedFilesInCL();
+	const TArray<FSCFileRef>& ShelvedFiles = ServiceProvider.Pin()->GetService<IChangelistService>()->GetShelvedFilesInCL();
 
 	if (ShelvedFiles.Num() != 0)
 	{
@@ -701,7 +712,7 @@ EDialogFactoryResult FPreflightService::ShowUpdatePreflightTagDialog() const
 		ModelInterface->GetMainTab().Pin()->GetParentWindow()->DrawAttention(FWindowDrawAttentionParameters());
 	}
 
-	const FText TextTitle = FText::FromString(FString::Printf(TEXT("Preflight CL %s: Newer preflight available"), *ServiceProvider.Pin()->GetService<FChangelistService>()->GetCLID()));
+	const FText TextTitle = FText::FromString(FString::Printf(TEXT("Preflight CL %s: Newer preflight available"), *ServiceProvider.Pin()->GetService<IChangelistService>()->GetCLID()));
 	const FText TextDescription = FText::FromString(FString::Printf(TEXT("There is a newer preflight for this changelist:\n<a id=\"browser\" style=\"Hyperlink\" href=\"%sjob/%s\">%s - %s</>\n\nDo you want to update the #preflight tag?"), *Definition.HordeServerAddress, *HordePreflights->PreflightList[0].ID, *HordePreflights->PreflightList[0].Name, *HordePreflights->PreflightList[0].ID));
 
 	return FDialogFactory::ShowDialog(TextTitle, TextDescription, TArray<FString>{ TEXT("Update Tag Value"), TEXT("Cancel") }, FSubmitToolUtils::BuildUserPrefCheckboxUI(FSubmitToolUserPrefs::Get()->bAutoUpdatePreflight, FText::FromString(TEXT("Always update, Don't ask again"))));
diff --git a/Engine/Source/Programs/SubmitTool/Private/Logic/PreflightService.h b/Engine/Source/Programs/SubmitTool/Private/Logic/PreflightService.h
--- a/Engine/Source/Programs/SubmitTool/Private/Logic/PreflightService.h
+++ b/Engine/Source/Programs/SubmitTool/Private/Logic/PreflightService.h
@@ -4,12 +4,12 @@
 
 #include "CoreMinimal.h"
 #include "Containers/Ticker.h"
-#include "ChangelistService.h"
+#include "Logic/Services/Interfaces/IChangelistService.h"
 #include "Logic/ProcessWrapper.h"
 #include "Logic/DialogFactory.h"
 #include "Parameters/SubmitToolParameters.h"
 #include "Models/PreflightData.h"
-#include "Services/Interfaces/ISubmitToolService.h"
+#include "Logic/Services/Interfaces/ISubmitToolService.h"
 
 class FSubmitToolServiceProvider;
 class FModelInterface;
diff --git a/Engine/Source/Programs/SubmitTool/Private/Logic/Services/CacheDataService.h b/Engine/Source/Programs/SubmitTool/Private/Logic/Services/CacheDataService.h
--- a/Engine/Source/Programs/SubmitTool/Private/Logic/Services/CacheDataService.h
+++ b/Engine/Source/Programs/SubmitTool/Private/Logic/Services/CacheDataService.h
@@ -4,7 +4,7 @@
 
 #include "CoreMinimal.h"
 #include "Models/CacheData.h"
-#include "Interfaces/ICacheDataService.h"
+#include "Logic/Services/Interfaces/ICacheDataService.h"
 
 struct FGeneralParameters;
 
diff --git a/Engine/Source/Programs/SubmitTool/Private/Logic/Services/Interfaces/ICacheDataService.h b/Engine/Source/Programs/SubmitTool/Private/Logic/Services/Interfaces/ICacheDataService.h
deleted file mode 100644
diff --git a/Engine/Source/Programs/SubmitTool/Private/Logic/Services/Interfaces/ISTSourceControlService.h b/Engine/Source/Programs/SubmitTool/Private/Logic/Services/Interfaces/ISTSourceControlService.h
deleted file mode 100644
diff --git a/Engine/Source/Programs/SubmitTool/Private/Logic/Services/Interfaces/ISubmitToolService.h b/Engine/Source/Programs/SubmitTool/Private/Logic/Services/Interfaces/ISubmitToolService.h
deleted file mode 100644
diff --git a/Engine/Source/Programs/SubmitTool/Private/Logic/Services/SubmitToolServiceProvider.h b/Engine/Source/Programs/SubmitTool/Private/Logic/Services/SubmitToolServiceProvider.h
deleted file mode 100644
diff --git a/Engine/Source/Programs/SubmitTool/Private/Logic/SwarmService.cpp b/Engine/Source/Programs/SubmitTool/Private/Logic/SwarmService.cpp
--- a/Engine/Source/Programs/SubmitTool/Private/Logic/SwarmService.cpp
+++ b/Engine/Source/Programs/SubmitTool/Private/Logic/SwarmService.cpp
@@ -5,8 +5,8 @@
 #include "Interfaces/IHttpResponse.h"
 #include "Logic/Services/SubmitToolServiceProvider.h"
 #include "Logic/Services/Interfaces/ISTSourceControlService.h"
-#include "Logic/TagService.h"
-#include "Logic/ChangelistService.h"
+#include "Logic/Services/Interfaces/ITagService.h"
+#include "Logic/Services/Interfaces/IChangelistService.h"
 #include "Serialization/JsonSerializer.h"
 #include "JsonObjectConverter.h"
 
@@ -24,7 +24,7 @@ FSwarmService::FSwarmService(TWeakPtr<FSubmitToolServiceProvider> InServiceProvi
 
 void FSwarmService::FetchReview(const OnGetReviewComplete& OnComplete)
 {
-	FString Changelist = ServiceProvider.Pin()->GetService<FChangelistService>()->GetCLID();
+	FString Changelist = ServiceProvider.Pin()->GetService<IChangelistService>()->GetCLID();
 
 	if (Changelist.IsEmpty())
 	{
@@ -85,7 +85,7 @@ void FSwarmService::FetchReview(const OnGetReviewComplete& OnComplete)
 
 void FSwarmService::CreateReview(const TArray<FString>& InReviewers, const OnCreateReviewComplete& OnComplete)
 {
-	FString Changelist = ServiceProvider.Pin()->GetService<FChangelistService>()->GetCLID();
+	FString Changelist = ServiceProvider.Pin()->GetService<IChangelistService>()->GetCLID();
 
 	if (Changelist.IsEmpty())
 	{
@@ -145,7 +145,7 @@ void FSwarmService::CreateReview(const TArray<FString>& InReviewers, const OnCre
 		RequestJson->SetArrayField(TEXT("reviewerGroups"), GroupsObject);
 	}
 
-	RequestJson->SetStringField(TEXT("description"), ServiceProvider.Pin()->GetService<FChangelistService>()->GetCLDescription());
+	RequestJson->SetStringField(TEXT("description"), ServiceProvider.Pin()->GetService<IChangelistService>()->GetCLDescription());
 
 	FString BodyString;
 	TSharedRef<TJsonWriter<TCHAR>> JsonWriter = TJsonWriterFactory<>::Create(&BodyString);
@@ -299,7 +299,7 @@ const TArray<FString> FSwarmService::GetUsersInSwarmTag() const
 {
 	TArray<FString> Reviewers;
 	
-	for (const FTag* Tag : ServiceProvider.Pin()->GetService<FTagService>()->GetTagsArray())
+	for (const FTag* Tag : ServiceProvider.Pin()->GetService<ITagService>()->GetTagsArray())
 	{
 		if (Tag->Definition.InputSubType.Equals(TEXT("Swarm"), ESearchCase::IgnoreCase))
 		{
diff --git a/Engine/Source/Programs/SubmitTool/Private/Logic/SwarmService.h b/Engine/Source/Programs/SubmitTool/Private/Logic/SwarmService.h
--- a/Engine/Source/Programs/SubmitTool/Private/Logic/SwarmService.h
+++ b/Engine/Source/Programs/SubmitTool/Private/Logic/SwarmService.h
@@ -5,8 +5,8 @@
 #include "CoreMinimal.h"
 
 #include "Models/SwarmReview.h"
-#include "Services/Interfaces/ISubmitToolService.h"
-#include "Services/Interfaces/ISTSourceControlService.h"
+#include "Logic/Services/Interfaces/ISubmitToolService.h"
+#include "Logic/Services/Interfaces/ISTSourceControlService.h"
 
 #include "HttpModule.h"
 
diff --git a/Engine/Source/Programs/SubmitTool/Private/Logic/TagService.cpp b/Engine/Source/Programs/SubmitTool/Private/Logic/TagService.cpp
--- a/Engine/Source/Programs/SubmitTool/Private/Logic/TagService.cpp
+++ b/Engine/Source/Programs/SubmitTool/Private/Logic/TagService.cpp
@@ -5,7 +5,7 @@
 #include "Logging/SubmitToolLog.h"
 
 
-FTagService::FTagService(const FSubmitToolParameters& InParameters, TSharedPtr<FChangelistService> CLService) :
+FTagService::FTagService(const FSubmitToolParameters& InParameters, TSharedPtr<IChangelistService> CLService) :
 	Parameters(InParameters),
 	ChangelistService(CLService)
 {
diff --git a/Engine/Source/Programs/SubmitTool/Private/Logic/TagService.h b/Engine/Source/Programs/SubmitTool/Private/Logic/TagService.h
--- a/Engine/Source/Programs/SubmitTool/Private/Logic/TagService.h
+++ b/Engine/Source/Programs/SubmitTool/Private/Logic/TagService.h
@@ -5,27 +5,27 @@
 #include "CoreMinimal.h"
 #include "Models/Tag.h"
 #include "Parameters/SubmitToolParameters.h"
-#include "ChangelistService.h"
-#include "Services/Interfaces/ISubmitToolService.h"
+#include "Logic/Services/Interfaces/IChangelistService.h"
+#include "Logic/Services/Interfaces/ITagService.h"
 
-class FTagService : public ISubmitToolService
+class FTagService : public ITagService
 {
 public:
-	FTagService(const FSubmitToolParameters& InParameters, TSharedPtr<FChangelistService> clService);
+	FTagService(const FSubmitToolParameters& InParameters, TSharedPtr<IChangelistService> clService);
 
-	void ParseCLDescription();
-	void ReApplyTags();
+	virtual void ParseCLDescription() override;
+	virtual void ReApplyTags() override;
 
-	void ApplyTag(const FString& tagID);
-	void ApplyTag(FTag& tag);
-	void UpdateTagsInCL();
-	void RemoveTag(const FString& tagID);
-	void RemoveTag(FTag& tag);
-	void SetTagValues(const FString& tagID, const FString& values);
-	void SetTagValues(FTag& tag, const FString& values);
-	void SetTagValues(FTag& tag, const TArray<FString>& values);
+	virtual void ApplyTag(const FString& tagID) override;
+	virtual void ApplyTag(FTag& tag) override;
+	virtual void UpdateTagsInCL() override;
+	virtual void RemoveTag(const FString& tagID) override;
+	virtual void RemoveTag(FTag& tag) override;
+	virtual void SetTagValues(const FString& tagID, const FString& values) override;
+	virtual void SetTagValues(FTag& tag, const FString& values) override;
+	virtual void SetTagValues(FTag& tag, const TArray<FString>& values) override;
 
-	FTag* GetTag(const FString& tagID)
+	virtual FTag* GetTag(const FString& tagID) override
 	{
 		if(RegisteredTags.Contains(tagID))
 		{
@@ -35,7 +35,7 @@ public:
 		return nullptr;
 	};
 
-	FTag* GetTagOfType(const FString& InType)
+	virtual FTag* GetTagOfType(const FString& InType) override
 	{
 		for (const TPair<FString, FTag>& Tag : RegisteredTags)
 		{
@@ -48,8 +48,7 @@ public:
 		return nullptr;
 	}
 
-
-	FTag* GetTagOfSubtype(const FString& InTagSubtype)
+	virtual FTag* GetTagOfSubtype(const FString& InTagSubtype) override
 	{
 		for (const TPair<FString, FTag>& Tag : RegisteredTags)
 		{
@@ -62,9 +61,7 @@ public:
 		return nullptr;
 	}
 
-	FTagUpdated OnTagUpdated;
-
-	const TArray<const FTag*>& GetTagsArray() const;
+	virtual const TArray<const FTag*>& GetTagsArray() const override;
 
 private:
 	FString& GetCLDescription() { return const_cast<FString&>(ChangelistService->GetCLDescription()); }
@@ -73,7 +70,7 @@ private:
 
 	const FSubmitToolParameters& Parameters;
 	TMap<FString, FTag> RegisteredTags;
-	TSharedPtr<FChangelistService> ChangelistService;
+	TSharedPtr<IChangelistService> ChangelistService;
 
 	void RegisterTags();
 	void UpdateTagsPositions(size_t changePos, int32 delta);
diff --git a/Engine/Source/Programs/SubmitTool/Private/Logic/TasksService.cpp b/Engine/Source/Programs/SubmitTool/Private/Logic/TasksService.cpp
--- a/Engine/Source/Programs/SubmitTool/Private/Logic/TasksService.cpp
+++ b/Engine/Source/Programs/SubmitTool/Private/Logic/TasksService.cpp
@@ -10,8 +10,6 @@
 #include "Policies/CondensedJsonPrintPolicy.h"
 #include "Serialization/JsonWriter.h"
 
-#include "ChangelistService.h"
-#include "TagService.h"
 
 
 FTasksService::FTasksService(const TMap<FString, FString>& InTasks, const FString& InTelemetryEventsId) :
diff --git a/Engine/Source/Programs/SubmitTool/Private/Logic/TasksService.h b/Engine/Source/Programs/SubmitTool/Private/Logic/TasksService.h
--- a/Engine/Source/Programs/SubmitTool/Private/Logic/TasksService.h
+++ b/Engine/Source/Programs/SubmitTool/Private/Logic/TasksService.h
@@ -5,7 +5,7 @@
 #include "CoreMinimal.h"
 #include "Logic/Validators/ValidatorBase.h"
 #include "Parameters/SubmitToolParameters.h"
-#include "Services/Interfaces/ISubmitToolService.h"
+#include "Logic/Services/Interfaces/ISubmitToolService.h"
 
 class FChangelistService;
 class FTagService;
diff --git a/Engine/Source/Programs/SubmitTool/Private/Logic/UpdateService.cpp b/Engine/Source/Programs/SubmitTool/Private/Logic/UpdateService.cpp
--- a/Engine/Source/Programs/SubmitTool/Private/Logic/UpdateService.cpp
+++ b/Engine/Source/Programs/SubmitTool/Private/Logic/UpdateService.cpp
@@ -16,7 +16,7 @@
 #include "Models/HordeDeploymentData.h"
 #include "Framework/Application/SlateApplication.h"
 #include "Logic/Services/SubmitToolServiceProvider.h"
-#include "Logic/ChangelistService.h"
+#include "Logic/Services/Interfaces/IChangelistService.h"
 
 FUpdateService::FUpdateService(const FHordeParameters& InHordeParameters, const FAutoUpdateParameters& InAutoUpdateParameters, TWeakPtr<FSubmitToolServiceProvider> InServiceProvider) :
 	HordeParameters(InHordeParameters),
@@ -425,7 +425,7 @@ FString FUpdateService::GetSubmitToolArgs() const
 
 	if (Commandline.Contains(TEXT("-cl default")))
 	{
-		Commandline.ReplaceInline(TEXT("-cl default"), *FString::Printf(TEXT("-cl %s"), *ServiceProvider.Pin()->GetService<FChangelistService>()->GetCLID()));
+		Commandline.ReplaceInline(TEXT("-cl default"), *FString::Printf(TEXT("-cl %s"), *ServiceProvider.Pin()->GetService<IChangelistService>()->GetCLID()));
 	}
 
 	return Commandline;
diff --git a/Engine/Source/Programs/SubmitTool/Private/Logic/UpdateService.h b/Engine/Source/Programs/SubmitTool/Private/Logic/UpdateService.h
--- a/Engine/Source/Programs/SubmitTool/Private/Logic/UpdateService.h
+++ b/Engine/Source/Programs/SubmitTool/Private/Logic/UpdateService.h
@@ -5,7 +5,7 @@
 #include "CoreMinimal.h"
 
 #include "Parameters/SubmitToolParameters.h"
-#include "Services/Interfaces/ISubmitToolService.h"
+#include "Logic/Services/Interfaces/ISubmitToolService.h"
 
 class IHttpRequest;
 class FSubmitToolServiceProvider;
diff --git a/Engine/Source/Programs/SubmitTool/Private/Logic/Validators/CrossChangelistValidator.cpp b/Engine/Source/Programs/SubmitTool/Private/Logic/Validators/CrossChangelistValidator.cpp
--- a/Engine/Source/Programs/SubmitTool/Private/Logic/Validators/CrossChangelistValidator.cpp
+++ b/Engine/Source/Programs/SubmitTool/Private/Logic/Validators/CrossChangelistValidator.cpp
@@ -4,7 +4,7 @@
 #include "Misc/Paths.h"
 #include "Misc/PathViews.h"
 #include "Logging/SubmitToolLog.h"
-#include "Logic/ChangelistService.h"
+#include "Logic/Services/Interfaces/IChangelistService.h"
 #include "CommandLine/CmdLineParameters.h"
 #include "Serialization/JsonSerializer.h"
 #include "SubmitToolCoreUtils.h"
@@ -74,7 +74,7 @@ bool FCrossChangelistValidator::CheckHeaderAndCppInDifferentChangelist()
 		return EExtension::Unknown;
 	};
 
-	TSharedPtr<FChangelistService> ChangelistService = ServiceProvider.Pin()->GetService<FChangelistService>();
+	TSharedPtr<IChangelistService> ChangelistService = ServiceProvider.Pin()->GetService<IChangelistService>();
 	const TArray<FSCFileRef>& FilesInChangelist = ChangelistService->GetFilesInCL();
 	TMap<FString, FString> OpenFilesInChangelists = ChangelistService->GetOpenFilesInChangelists();
 	FString CLID = ChangelistService->GetCLID();
diff --git a/Engine/Source/Programs/SubmitTool/Private/Logic/Validators/PreflightValidator.cpp b/Engine/Source/Programs/SubmitTool/Private/Logic/Validators/PreflightValidator.cpp
--- a/Engine/Source/Programs/SubmitTool/Private/Logic/Validators/PreflightValidator.cpp
+++ b/Engine/Source/Programs/SubmitTool/Private/Logic/Validators/PreflightValidator.cpp
@@ -3,7 +3,7 @@
 #include "PreflightValidator.h"
 
 #include "Logic/PreflightService.h"
-#include "Logic/TagService.h"
+#include "Logic/Services/Interfaces/ITagService.h"
 #include "Logic/Validators/ValidatorFactory.h"
 #include "Logging/SubmitToolLog.h"
 #include "Internationalization/Regex.h"
@@ -23,7 +23,7 @@ FPreflightValidator::FPreflightValidator(const FName& InNameId, const FSubmitToo
 bool FPreflightValidator::Activate()
 {
 	bIsValidSetup = FValidatorBase::Activate();
-	PreflightTag = ServiceProvider.Pin()->GetService<FTagService>()->GetTagOfSubtype(TEXT("preflight"));
+	PreflightTag = ServiceProvider.Pin()->GetService<ITagService>()->GetTagOfSubtype(TEXT("preflight"));
 
 	bIsValidSetup = bIsValidSetup && PreflightTag != nullptr;
 
@@ -180,7 +180,7 @@ const FString FPreflightValidator::GetPopupMessageWhenFailedText() const
 					if(Definition->bInvalidatesWhenOutOfDate)
 					{
 						IPlatformFile& PlatformFile = FPlatformFileManager::Get().GetPlatformFile();
-						for(const FSCFileRef& file : ServiceProvider.Pin()->GetService<FChangelistService>()->GetFilesInCL())
+						for(const FSCFileRef& file : ServiceProvider.Pin()->GetService<IChangelistService>()->GetFilesInCL())
 						{
 							const FDateTime FileTimestampUTC = PlatformFile.GetTimeStamp(*file->GetFilename());
 							if(FileTimestampUTC > FoundData->CreateTime)
diff --git a/Engine/Source/Programs/SubmitTool/Private/Logic/Validators/TagValidator.cpp b/Engine/Source/Programs/SubmitTool/Private/Logic/Validators/TagValidator.cpp
--- a/Engine/Source/Programs/SubmitTool/Private/Logic/Validators/TagValidator.cpp
+++ b/Engine/Source/Programs/SubmitTool/Private/Logic/Validators/TagValidator.cpp
@@ -4,7 +4,7 @@
 #include "Models/Tag.h"
 #include "Logging/SubmitToolLog.h"
 #include "Logic/Services/Interfaces/ISTSourceControlService.h"
-#include "Logic/ChangelistService.h"
+#include "Logic/Services/Interfaces/IChangelistService.h"
 #include "Logic/Validators/ValidatorFactory.h"
 #include "Internationalization/Regex.h"
 
@@ -73,7 +73,7 @@ bool FTagValidator::ValidateTag(const FTag* InTag, const TArray<TSharedPtr<FUser
 {
 	bool bIsTagValid = true;
 
-	FTagValidationConfig ValidationConfig = InTag->GetCurrentValidationConfig(ServiceProvider.Pin()->GetService<FChangelistService>()->GetFilesInCL());
+	FTagValidationConfig ValidationConfig = InTag->GetCurrentValidationConfig(ServiceProvider.Pin()->GetService<IChangelistService>()->GetFilesInCL());
 
 	if(ValidationConfig.bIsMandatory && !InTag->IsEnabled())
 	{
diff --git a/Engine/Source/Programs/SubmitTool/Private/Logic/Validators/ValidatorBase.cpp b/Engine/Source/Programs/SubmitTool/Private/Logic/Validators/ValidatorBase.cpp
--- a/Engine/Source/Programs/SubmitTool/Private/Logic/Validators/ValidatorBase.cpp
+++ b/Engine/Source/Programs/SubmitTool/Private/Logic/Validators/ValidatorBase.cpp
@@ -10,7 +10,7 @@
 #include "CommandLine/CmdLineParameters.h"
 #include "Configuration/Configuration.h"
 #include "Logic/DialogFactory.h"
-#include "Logic/TagService.h"
+#include "Logic/Services/Interfaces/ITagService.h"
 #include "Logic/PreflightService.h"
 #include "Logic/Services/Interfaces/ICacheDataService.h"
 #include "Logic/Services/Interfaces/ISTSourceControlService.h"
@@ -80,8 +80,8 @@ void FValidatorBase::StartValidation()
 	WarningListCache.Reset();
 	CorrelationId = FGuid::NewGuid();
 
-	const TSharedPtr<FChangelistService>& ChangelistService = ServiceProvider.Pin()->GetService<FChangelistService>();
-	const TSharedPtr<FTagService>& TagService = ServiceProvider.Pin()->GetService<FTagService>();
+	const TSharedPtr<IChangelistService>& ChangelistService = ServiceProvider.Pin()->GetService<IChangelistService>();
+	const TSharedPtr<ITagService>& TagService = ServiceProvider.Pin()->GetService<ITagService>();
 
 	FilteredFiles.RemoveAt(0, FilteredFiles.Num(), EAllowShrinking::No);
 
@@ -248,7 +248,7 @@ void FValidatorBase::Tick(float InDeltaTime)
 
 	if (Definition->TimeoutLimit > 0 && RunTime >= Definition->TimeoutLimit)
 	{
-		LogFailure(FString::Printf(TEXT("[%s] %s"), *GetValidatorName(), TEXT("Timeout limit has been reached, cancelling task.")));
+		LogFailure(FString::Printf(TEXT("[%s] Timeout limit has been reached (%.0f seconds), cancelling task."), *GetValidatorName(), Definition->TimeoutLimit));
 
 		StopInternalValidations();
 		SetValidationFinished(EValidationStates::Timeout);
@@ -277,7 +277,7 @@ bool FValidatorBase::CanStartTask() const
 	bool bCanRun = true;
 	if (Definition->bRequireShelveOperation)
 	{		
-		bCanRun &= ServiceProvider.Pin()->GetService<FChangelistService>()->IsShelfReady();
+		bCanRun &= ServiceProvider.Pin()->GetService<IChangelistService>()->IsShelfReady();
 	}
 
 	if (!Definition->PreflightTemplate.IsEmpty())
@@ -313,7 +313,7 @@ void FValidatorBase::InvalidateLocalFileModifications()
 	if ((Definition->TaskArea & ETaskArea::LocalFiles) == ETaskArea::LocalFiles && (State == EValidationStates::Valid || State == EValidationStates::Running || State == EValidationStates::Skipped || State == EValidationStates::Not_Applicable))
 	{
 		FFileManagerGeneric FileManager;
-		for (const FSCFileRef& File : ServiceProvider.Pin()->GetService<FChangelistService>()->GetFilesInCL())
+		for (const FSCFileRef& File : ServiceProvider.Pin()->GetService<IChangelistService>()->GetFilesInCL())
 		{
 			bool bIncrementallySkipped;
 
@@ -351,7 +351,7 @@ const FString FValidatorBase::GetStatusText() const
 
 	if (State == EValidationStates::Queued)
 	{
-		if (Definition->bRequireShelveOperation && !ServiceProvider.Pin()->GetService<FChangelistService>()->IsShelfReady())
+		if (Definition->bRequireShelveOperation && !ServiceProvider.Pin()->GetService<IChangelistService>()->IsShelfReady())
 		{
 			AdditionalInformation = TEXT("Waiting Shelf");
 		}
@@ -384,7 +384,7 @@ const FString FValidatorBase::GetStatusText() const
 
 const TArray<FAnalyticsEventAttribute> FValidatorBase::GetTelemetryAttributes() const
 {
-	const TSharedPtr<FChangelistService> ChangelistService = ServiceProvider.Pin()->GetService<FChangelistService>();
+	const TSharedPtr<IChangelistService> ChangelistService = ServiceProvider.Pin()->GetService<IChangelistService>();
 
 	FString StealthStateStr = TEXT("No Stealth Mode");
 
@@ -449,7 +449,7 @@ TArray<FAnalyticsEventAttribute> FValidatorBase::GetSingleMessageTelemetryAttrib
 		TEXT("IsStealthValidator"), Definition->bIsStealthMode,
 		TEXT("Runtime"), RunTime,
 		TEXT("Stream"), ServiceProvider.Pin()->GetService<ISTSourceControlService>()->GetCurrentStreamName(),
-		TEXT("PendingChangelist"), ServiceProvider.Pin()->GetService<FChangelistService>()->GetCLID(),
+		TEXT("PendingChangelist"), ServiceProvider.Pin()->GetService<IChangelistService>()->GetCLID(),
 		TEXT("FromEditor"), FCmdLineParameters::Get().Contains(FSubmitToolCmdLine::EditorFlag),
 		TEXT("Message"), InMessage
 	);
@@ -484,7 +484,7 @@ void FValidatorBase::ValidationFinished(const bool bPassed)
 
 	if (bPassed && Definition->bUsesIncrementalCache)
 	{
-		ServiceProvider.Pin()->GetService<ICacheDataService>()->UpdateLastValidationForFiles(ServiceProvider.Pin()->GetService<FChangelistService>()->GetCLID(), GetValidatorNameId(), GetValidationConfigId(), FilteredFiles, FDateTime::UtcNow());
+		ServiceProvider.Pin()->GetService<ICacheDataService>()->UpdateLastValidationForFiles(ServiceProvider.Pin()->GetService<IChangelistService>()->GetCLID(), GetValidatorNameId(), GetValidationConfigId(), FilteredFiles, FDateTime::UtcNow());
 	}
 
 	EValidationStates NewState = bPassed ? EValidationStates::Valid : EValidationStates::Failed;
@@ -563,7 +563,7 @@ void FValidatorBase::SetValidationFinished(EValidationStates NewState)
 
 bool FValidatorBase::EvaluateTagSkip()
 {
-	TSharedPtr<FChangelistService> ChangelistService = ServiceProvider.Pin()->GetService<FChangelistService>();
+	TSharedPtr<IChangelistService> ChangelistService = ServiceProvider.Pin()->GetService<IChangelistService>();
 	if (Definition->SkipForbiddenFiles.Num() > 0)
 	{
 		for (const FString& File : Definition->SkipForbiddenFiles)
@@ -638,11 +638,11 @@ void FValidatorBase::Skip()
 
 bool FValidatorBase::IsRelevantToCL() const
 {
-	TSharedPtr<FChangelistService> ChangelistService = ServiceProvider.Pin()->GetService<FChangelistService>();
+	TSharedPtr<IChangelistService> ChangelistService = ServiceProvider.Pin()->GetService<IChangelistService>();
 
 	TArray<FSCFileRef> IncrementallySkippedFiles;
 	TArray<FSCFileRef> OutFiles;
-	return AppliesToCL(ChangelistService->GetCLDescription(), ChangelistService->GetFilesInCL(), ServiceProvider.Pin()->GetService<FTagService>()->GetTagsArray(), OutFiles, IncrementallySkippedFiles, false);
+	return AppliesToCL(ChangelistService->GetCLDescription(), ChangelistService->GetFilesInCL(), ServiceProvider.Pin()->GetService<ITagService>()->GetTagsArray(), OutFiles, IncrementallySkippedFiles, false);
 }
 
 void FValidatorBase::SetSelectedOption(const FString& InOptionName, const FString& InOptionValue)
@@ -828,7 +828,7 @@ bool FValidatorBase::AppliesToFile(const FSCFileRef InFile, bool InbAllowIncreme
 	if (InbAllowIncremental && bIncluded)
 	{
 		FFileManagerGeneric FileManager;
-		FDateTime LastValidation = ServiceProvider.Pin()->GetService<ICacheDataService>()->GetLastValidationDate(ServiceProvider.Pin()->GetService<FChangelistService>()->GetCLID(), GetValidatorNameId(), GetValidationConfigId(), Filename);
+		FDateTime LastValidation = ServiceProvider.Pin()->GetService<ICacheDataService>()->GetLastValidationDate(ServiceProvider.Pin()->GetService<IChangelistService>()->GetCLID(), GetValidatorNameId(), GetValidationConfigId(), Filename);
 		FFileStatData FileModifiedDate = FileManager.GetStatData(*Filename);
 		if (LastValidation != FDateTime::MinValue() && FileModifiedDate.ModificationTime < LastValidation)
 		{
diff --git a/Engine/Source/Programs/SubmitTool/Private/Models/ModelInterface.cpp b/Engine/Source/Programs/SubmitTool/Private/Models/ModelInterface.cpp
--- a/Engine/Source/Programs/SubmitTool/Private/Models/ModelInterface.cpp
+++ b/Engine/Source/Programs/SubmitTool/Private/Models/ModelInterface.cpp
@@ -1,6 +1,10 @@
 // Copyright Epic Games, Inc. All Rights Reserved.
 
 #include "ModelInterface.h"
+#include "Logic/ChangelistService.h"
+#include "Logic/TagService.h"
+#include "Logic/P4LockdownService.h"
+#include "Logic/CredentialsService.h"
 #include "Algo/Compare.h"
 #include "Models/SubmitToolUserPrefs.h"
 #include "Logic/JiraService.h"
@@ -61,13 +65,13 @@ FModelInterface::FModelInterface(FSubmitToolParameters&& InParameters, const TSh
 	PresubmitOperationsService = MakeShared<FTasksService>(Parameters.PresubmitOperations, TEXT("SubmitTool.StandAlone.PresubmitOperation"));
 	ServiceProvider->RegisterService<FTasksService>(PresubmitOperationsService.ToSharedRef(), TEXT("PresubmitOperationsService"));
 	CredentialsService = MakeShared<FCredentialsService>(Parameters);
-	ServiceProvider->RegisterService<FCredentialsService>(CredentialsService.ToSharedRef());
+	ServiceProvider->RegisterService<ICredentialsService>(CredentialsService.ToSharedRef());
 	ChangelistService = MakeShared<FChangelistService>(Parameters.GeneralParameters, ServiceProvider, CLReadyCallback, CLRefreshCallback);
-	ServiceProvider->RegisterService<FChangelistService>(ChangelistService.ToSharedRef());
+	ServiceProvider->RegisterService<IChangelistService>(ChangelistService.ToSharedRef());
 	P4LockdownService = MakeShared<FP4LockdownService>(Parameters.P4LockdownParameters, ServiceProvider);
-	ServiceProvider->RegisterService<FP4LockdownService>(P4LockdownService.ToSharedRef());
+	ServiceProvider->RegisterService<ILockdownService>(P4LockdownService.ToSharedRef());
 	TagService = MakeShared<FTagService>(Parameters, ChangelistService);
-	ServiceProvider->RegisterService<FTagService>(TagService.ToSharedRef());
+	ServiceProvider->RegisterService<ITagService>(TagService.ToSharedRef());
 	SwarmService = MakeShared<FSwarmService>(ServiceProvider);
 	ServiceProvider->RegisterService<FSwarmService>(SwarmService.ToSharedRef());
 	PreflightService = MakeShared<FPreflightService>(Parameters.HordeParameters, this, ServiceProvider);
diff --git a/Engine/Source/Programs/SubmitTool/Private/Models/ModelInterface.h b/Engine/Source/Programs/SubmitTool/Private/Models/ModelInterface.h
--- a/Engine/Source/Programs/SubmitTool/Private/Models/ModelInterface.h
+++ b/Engine/Source/Programs/SubmitTool/Private/Models/ModelInterface.h
@@ -3,16 +3,16 @@
 #pragma once
 
 #include "CoreMinimal.h"
-#include "Logic/ChangelistService.h"
 #include "Logic/TasksService.h"
-#include "Logic/TagService.h"
 #include "Logic/Services/Interfaces/ISTSourceControlService.h"
+#include "Logic/Services/Interfaces/IChangelistService.h"
+#include "Logic/Services/Interfaces/ITagService.h"
+#include "Logic/Services/Interfaces/ILockdownService.h"
+#include "Logic/Services/Interfaces/ICredentialsService.h"
 #include "Logic/JiraService.h"
 #include "Logic/IntegrationService.h"
 #include "Logic/SwarmService.h"
 #include "Logic/PreflightService.h"
-#include "Logic/P4LockdownService.h"
-#include "Logic/CredentialsService.h"
 #include "Logic/UpdateService.h"
 #include "Logic/Services/SubmitToolServiceProvider.h"
 #include "Models/Tag.h"
@@ -279,9 +279,9 @@ private:
 	FString GeneratedDescription;
 
 	TSharedPtr<ISTSourceControlService> SourceControlService;
-	TSharedPtr<FChangelistService> ChangelistService;
-	TSharedPtr<FP4LockdownService> P4LockdownService;
-	TSharedPtr<FTagService> TagService;
+	TSharedPtr<IChangelistService> ChangelistService;
+	TSharedPtr<ILockdownService> P4LockdownService;
+	TSharedPtr<ITagService> TagService;
 	TSharedPtr<FTasksService> ValidationService;
 	TSharedPtr<FJiraService> JiraService;
 	TSharedPtr<FPreflightService> PreflightService;
@@ -290,7 +290,7 @@ private:
 	TSharedPtr<FTasksService> PresubmitOperationsService;
 	TSharedPtr<FIntegrationService> IntegrationService;
 	TSharedPtr<FSwarmService> SwarmService;
-	TSharedPtr<FCredentialsService> CredentialsService;
+	TSharedPtr<ICredentialsService> CredentialsService;
 	TSharedPtr<FUpdateService> UpdateService;
 	TSharedPtr<FSubmitToolServiceProvider> ServiceProvider;
 	bool bPreflightQueued = false;
diff --git a/Engine/Source/Programs/SubmitTool/Private/SubmitToolApp.cpp b/Engine/Source/Programs/SubmitTool/Private/SubmitToolApp.cpp
--- a/Engine/Source/Programs/SubmitTool/Private/SubmitToolApp.cpp
+++ b/Engine/Source/Programs/SubmitTool/Private/SubmitToolApp.cpp
@@ -225,6 +225,9 @@ int RunSubmitTool(const TCHAR* CommandLine)
 	// Tell the module manager it may now process newly-loaded UObjects when new C++ modules are loaded
 	FModuleManager::Get().StartProcessingNewlyLoadedObjects();
 
+	// Load optional restricted module if present — silently no-ops if not found
+	FModuleManager::Get().LoadModule(TEXT("SubmitToolRestricted"));
+
 	// initialize the configuration system
 	FConfiguration::Init();
 
diff --git a/Engine/Source/Programs/SubmitTool/Private/Version/AppVersionDefines.h b/Engine/Source/Programs/SubmitTool/Private/Version/AppVersionDefines.h
--- a/Engine/Source/Programs/SubmitTool/Private/Version/AppVersionDefines.h
+++ b/Engine/Source/Programs/SubmitTool/Private/Version/AppVersionDefines.h
@@ -20,8 +20,8 @@
 #endif
 
 #define SUBMIT_TOOL_VERSION_MAJOR 6
-#define SUBMIT_TOOL_VERSION_MINOR 7
-#define SUBMIT_TOOL_VERSION_PATCH 3
+#define SUBMIT_TOOL_VERSION_MINOR 8
+#define SUBMIT_TOOL_VERSION_PATCH 0
 
 #ifndef SUBMIT_TOOL_CHANGELIST
 #define SUBMIT_TOOL_CHANGELIST_STRING TEXT("NoCL")
diff --git a/Engine/Source/Programs/SubmitTool/SubmitTool.Build.cs b/Engine/Source/Programs/SubmitTool/SubmitTool.Build.cs
--- a/Engine/Source/Programs/SubmitTool/SubmitTool.Build.cs
+++ b/Engine/Source/Programs/SubmitTool/SubmitTool.Build.cs
@@ -67,6 +67,14 @@ public class SubmitTool : ModuleRules
 			PublicDependencyModuleNames.Add("RPCBase");
 		}
 
+		// Include the restricted module if it exists
+		string RestrictedModulePath = Path.Combine(EngineDirectory,
+			"Restricted", "NotForLicensees", "Source", "Programs", "SubmitToolRestricted", "SubmitToolRestricted.Build.cs");
+		if (File.Exists(RestrictedModulePath))
+		{
+			PrivateDependencyModuleNames.Add("SubmitToolRestricted");
+		}
+
 		AddEngineThirdPartyPrivateStaticDependencies(Target, "Perforce");
 
 		if (Target.IsInPlatformGroup(UnrealPlatformGroup.Linux))
